Updated ML integration to MLSDK 0.19

diff --git a/Engine/Plugins/Lumin/MagicLeapScreens/Source/MagicLeapScreens/Private/MagicLeapScreensMsg.h b/Engine/Plugins/Lumin/MagicLeapScreens/Source/MagicLeapScreens/Private/MagicLeapScreensMsg.h
new file mode 100644
--- /dev/null
+++ b/Engine/Plugins/Lumin/MagicLeapScreens/Source/MagicLeapScreens/Private/MagicLeapScreensMsg.h
@@ -0,0 +1,36 @@
+// Copyright 1998-2018 Epic Games, Inc. All Rights Reserved.
+
+#pragma once
+
+#include "Containers/Array.h"
+#include "MagicLeapScreensTypes.h"
+
+enum EScreensMsgType
+{
+	Request,
+	Response,
+};
+
+enum EScreensTaskType
+{
+	None,
+	GetHistory,
+	AddToHistory,
+	UpdateEntry,
+};
+
+struct FScreensMessage
+{
+	EScreensMsgType Type;
+	EScreensTaskType TaskType;
+	bool bSuccess;
+	TArray<FScreensWatchHistoryEntry> WatchHistory;
+	FScreensEntryRequestResultDelegate EntryDelegate;
+	FScreensHistoryRequestResultDelegate HistoryDelegate;
+
+	FScreensMessage()
+		: Type(EScreensMsgType::Request)
+		, TaskType(EScreensTaskType::None)
+		, bSuccess(false)
+	{}
+};
\ No newline at end of file
diff --git a/Engine/Plugins/Lumin/MagicLeapScreens/Source/MagicLeapScreens/Private/MagicLeapScreensPlugin.cpp b/Engine/Plugins/Lumin/MagicLeapScreens/Source/MagicLeapScreens/Private/MagicLeapScreensPlugin.cpp
new file mode 100644
--- /dev/null
+++ b/Engine/Plugins/Lumin/MagicLeapScreens/Source/MagicLeapScreens/Private/MagicLeapScreensPlugin.cpp
@@ -0,0 +1,503 @@
+// Copyright 1998-2018 Epic Games, Inc. All Rights Reserved.
+
+#include "MagicLeapScreensPlugin.h"
+#include "RenderUtils.h"
+#include "IMagicLeapPlugin.h"
+#include "HeadMountedDisplayFunctionLibrary.h"
+#include "MagicLeapScreensWorker.h"
+#include "MagicLeapHMD.h"
+
+#define MAX_TEXTURE_SIZE 450 * 450 * 4 // currently limited by binder implementation
+
+DEFINE_LOG_CATEGORY_STATIC(LogScreensPlugin, Display, All);
+
+#if WITH_MLSDK
+MLImage FMagicLeapScreensPlugin::DefaultThumbnail = MLImage();
+FCriticalSection FMagicLeapScreensPlugin::CriticalSection;
+#endif // WITH_MLSDK
+FScreensWorker *FMagicLeapScreensPlugin::Impl = new FScreensWorker();
+TArray<uint8> FMagicLeapScreensPlugin::PixelDataMemPool = TArray<uint8>();
+
+void FMagicLeapScreensPlugin::StartupModule()
+{
+	IModuleInterface::StartupModule();
+	APISetup.Startup();
+	APISetup.LoadDLL(TEXT("ml_screens"));
+
+#if WITH_MLSDK
+	DefaultThumbnail.width = 2;
+	DefaultThumbnail.height = 2;
+	DefaultThumbnail.image_type = MLImageType_RGBA32;
+	DefaultThumbnail.alignment = 1;
+	const SIZE_T DataSize = DefaultThumbnail.width * DefaultThumbnail.height * 4;
+	DefaultThumbnail.data = new uint8[DataSize];
+	FMemory::Memset(DefaultThumbnail.data, 255, DataSize);
+#endif // WITH_MLSDK
+
+	TickDelegate = FTickerDelegate::CreateRaw(this, &FMagicLeapScreensPlugin::Tick);
+	TickDelegateHandle = FTicker::GetCoreTicker().AddTicker(TickDelegate);
+
+	PixelDataMemPool.Reserve(MAX_TEXTURE_SIZE);
+	bEngineLoopInitComplete = false;
+	FCoreDelegates::OnFEngineLoopInitComplete.AddRaw(this, &FMagicLeapScreensPlugin::OnEngineLoopInitComplete);
+}
+
+void FMagicLeapScreensPlugin::ShutdownModule()
+{
+	FCoreDelegates::OnFEngineLoopInitComplete.RemoveAll(this);
+	APISetup.Shutdown();
+
+#if PLATFORM_LUMIN
+	if (Impl)
+	{
+		Impl->AsyncDestroy();
+		Impl = nullptr;
+	}
+#else
+	delete Impl;
+	Impl = nullptr;
+#endif
+
+#if WITH_MLSDK
+	delete[] DefaultThumbnail.data;
+	DefaultThumbnail.data = nullptr;
+	MLScreensReleaseWatchHistoryThumbnail(&FMagicLeapScreensPlugin::DefaultThumbnail);
+#endif // WITH_MLSDK
+	FTicker::GetCoreTicker().RemoveTicker(TickDelegateHandle);
+	IModuleInterface::ShutdownModule();
+}
+
+bool FMagicLeapScreensPlugin::IsEngineLoopInitComplete() const
+{
+	return bEngineLoopInitComplete;
+}
+
+void FMagicLeapScreensPlugin::OnEngineLoopInitComplete()
+{
+	if (Impl->Semaphore == nullptr)
+	{
+		Impl->EngineInited();
+	}
+
+	bEngineLoopInitComplete = true;
+}
+
+bool FMagicLeapScreensPlugin::IsSupportedFormat(EPixelFormat InPixelFormat)
+{
+	if (InPixelFormat == PF_B8G8R8A8 || InPixelFormat == PF_R8G8B8A8)
+	{
+		return true;
+	}
+
+	UE_LOG(LogScreensPlugin, Error, TEXT("Unsupported pixel format!"));
+
+	return false;
+}
+
+#if WITH_MLSDK
+UTexture2D* FMagicLeapScreensPlugin::MLImageToUTexture2D(const MLImage& Source)
+{
+	UTexture2D* Thumbnail = UTexture2D::CreateTransient(Source.width, Source.height, EPixelFormat::PF_R8G8B8A8);
+	FTexture2DMipMap& Mip = Thumbnail->PlatformData->Mips[0];
+	void* PixelData = Mip.BulkData.Lock(LOCK_READ_WRITE);
+	const uint32 PixelDataSize = Mip.BulkData.GetBulkDataSize();
+	FMemory::Memcpy(PixelData, Source.data, PixelDataSize);
+	UE_LOG(LogScreensPlugin, Log, TEXT("MLImageToUTexture2D width = %u height = %u size = %u"), Source.width, Source.height, PixelDataSize);
+	Thumbnail->SRGB = true;
+	Mip.BulkData.Unlock();
+	Thumbnail->UpdateResource();
+
+	return Thumbnail;
+}
+
+void FMagicLeapScreensPlugin::MLWatchHistoryEntryToUnreal(const MLScreensWatchHistoryEntry& InEntry, FScreensWatchHistoryEntry& OutEntry)
+{
+	OutEntry.ID.ID = InEntry.id;
+	OutEntry.Title = FString(UTF8_TO_TCHAR(InEntry.title));
+	OutEntry.Subtitle = FString(UTF8_TO_TCHAR(InEntry.subtitle));
+	OutEntry.PlaybackPosition = FTimespan(InEntry.playback_position_ms * ETimespan::TicksPerMillisecond);
+	OutEntry.PlaybackDuration = FTimespan(InEntry.playback_duration_ms * ETimespan::TicksPerMillisecond);
+	OutEntry.CustomData = FString(UTF8_TO_TCHAR(InEntry.custom_data));
+
+	MLImage MLThumbnail;
+	MLResult Result = MLScreensGetWatchHistoryThumbnail(InEntry.id, &MLThumbnail);
+	if (Result != MLResult_Ok)
+	{
+		UE_LOG(LogScreensPlugin, Log, TEXT("MLScreensGetWatchHistoryThumbnail failed for screen ID %u with error %s!"), (uint32)InEntry.id, UTF8_TO_TCHAR(MLScreensGetResultString(Result)));
+		OutEntry.Thumbnail = MLImageToUTexture2D(FMagicLeapScreensPlugin::DefaultThumbnail);
+	}
+	else
+	{
+		OutEntry.Thumbnail = MLImageToUTexture2D(MLThumbnail);
+		// Only release when default thumbnail is not used, default thumbnail should only be released when plugin shuts down
+		Result = MLScreensReleaseWatchHistoryThumbnail(&MLThumbnail);
+		UE_CLOG(Result != MLResult_Ok, LogScreensPlugin, Error, TEXT("MLScreensReleaseWatchHistoryThumbnail failed for with error %s!"), UTF8_TO_TCHAR(MLScreensGetResultString(Result)));
+	}
+}
+
+bool FMagicLeapScreensPlugin::UTexture2DToMLImage(const UTexture2D& Source, MLImage& Target)
+{
+	bool bSuccess = false;
+	FTexture2DMipMap& Mip = Source.PlatformData->Mips[0];
+	void* PixelData = Mip.BulkData.Lock(LOCK_READ_ONLY);
+	const int32 size = Mip.BulkData.GetBulkDataSize();
+
+	if ((size <= MAX_TEXTURE_SIZE) && (size != 0) && (PixelData != nullptr))
+	{
+		UE_LOG(LogScreensPlugin, Log, TEXT("UTexture2DToMLImage width = %d height = %d size = %d"), Mip.SizeX, Mip.SizeY, size);
+		Target.width = Mip.SizeX;
+		Target.height = Mip.SizeY;
+		Target.image_type = MLImageType_RGBA32;
+		Target.alignment = 1;
+		Target.data = PixelDataMemPool.GetData();
+		FMemory::Memcpy(Target.data, PixelData, size);
+
+		if (Source.GetPixelFormat() == EPixelFormat::PF_B8G8R8A8)
+		{
+			check(((size % 4) == 0) && (size != 0));
+
+			for (int32 i = 0; i < size - 4; i += 4)
+			{
+				Swap<uint8>(Target.data[i], Target.data[i + 2]);
+			}
+		}
+
+		bSuccess = true;
+	}
+	else
+	{
+		UE_CLOG(size > MAX_TEXTURE_SIZE, LogScreensPlugin, Error, TEXT("Texture size (%d) exceeds max texture size (%d)"), size, MAX_TEXTURE_SIZE);
+		UE_CLOG(size == 0, LogScreensPlugin, Error, TEXT("Texture size is zero"));
+	}
+
+	Mip.BulkData.Unlock();
+
+	return bSuccess;
+}
+
+bool FMagicLeapScreensPlugin::ShouldUseDefaultThumbnail(const FScreensWatchHistoryEntry& Entry,  MLImage& MLImage)
+{
+	return ((Entry.Thumbnail == nullptr) ||
+		(FMagicLeapScreensPlugin::IsSupportedFormat(Entry.Thumbnail->GetPixelFormat()) == false) ||
+		(FMagicLeapScreensPlugin::UTexture2DToMLImage(*Entry.Thumbnail, MLImage) == false)) ? true : false;
+}
+#endif // WITH_MLSDK
+
+bool FMagicLeapScreensPlugin::RemoveWatchHistoryEntry(const FScreenID& ID)
+{
+#if WITH_MLSDK
+	MLResult Result = MLResult_UnspecifiedFailure;
+	{
+		FScopeLock Lock(&FMagicLeapScreensPlugin::CriticalSection);
+		Result = MLScreensRemoveWatchHistoryEntry(ID.ID);
+		UE_CLOG(Result != MLResult_Ok, LogScreensPlugin, Error, TEXT("MLScreensRemoveWatchHistoryEntry failed with error %d for entry with id %d!"), Result, ID.ID);
+	}
+	return Result == MLResult_Ok;
+#else
+	return false;
+#endif // WITH_MLSDK
+}
+
+void FMagicLeapScreensPlugin::GetWatchHistoryEntriesAsync(const TOptional<FScreensHistoryRequestResultDelegate>& OptionalResultDelegate)
+{
+	FScreensMessage Msg;
+	Msg.Type = EScreensMsgType::Request;
+	Msg.TaskType = EScreensTaskType::GetHistory;
+	if (OptionalResultDelegate.IsSet())
+	{
+		Msg.HistoryDelegate = OptionalResultDelegate.GetValue();
+	}
+	FMagicLeapScreensPlugin::Impl->ProcessMessage(Msg);
+}
+
+void FMagicLeapScreensPlugin::AddToWatchHistoryAsync(const FScreensWatchHistoryEntry& NewEntry, const TOptional<FScreensEntryRequestResultDelegate>& OptionalResultDelegate)
+{
+	FScreensMessage Msg;
+	Msg.Type = EScreensMsgType::Request;
+	Msg.TaskType = EScreensTaskType::AddToHistory;
+	if (OptionalResultDelegate.IsSet())
+	{
+		Msg.EntryDelegate = OptionalResultDelegate.GetValue();
+	}
+	Msg.WatchHistory.Add(NewEntry);
+	FMagicLeapScreensPlugin::Impl->ProcessMessage(Msg);
+}
+
+void FMagicLeapScreensPlugin::UpdateWatchHistoryEntryAsync(const FScreensWatchHistoryEntry& UpdateEntry, const TOptional<FScreensEntryRequestResultDelegate>& OptionalResultDelegate)
+{
+	FScreensMessage Msg;
+	Msg.Type = EScreensMsgType::Request;
+	Msg.TaskType = EScreensTaskType::UpdateEntry;
+	if (OptionalResultDelegate.IsSet())
+	{
+		Msg.EntryDelegate = OptionalResultDelegate.GetValue();
+	}
+	Msg.WatchHistory.Add(UpdateEntry);
+	FMagicLeapScreensPlugin::Impl->ProcessMessage(Msg);
+}
+
+
+FScreensMessage FMagicLeapScreensPlugin::GetWatchHistoryEntries()
+{
+	FScreensMessage Msg;
+	Msg.Type = EScreensMsgType::Response;
+	Msg.TaskType = EScreensTaskType::GetHistory;
+#if WITH_MLSDK
+	{
+		FScopeLock Lock(&FMagicLeapScreensPlugin::CriticalSection);
+		MLScreensWatchHistoryList WatchHistoryList;
+		MLResult Result = MLScreensGetWatchHistoryList(&WatchHistoryList);
+		if (Result == MLResult_Ok)
+		{
+			Msg.WatchHistory.Reserve(WatchHistoryList.count);
+			for (uint32 i = 0; i < WatchHistoryList.count; ++i)
+			{
+				FScreensWatchHistoryEntry WatchHistoryEntry;
+				FMagicLeapScreensPlugin::MLWatchHistoryEntryToUnreal(WatchHistoryList.entries[i], WatchHistoryEntry);
+				Msg.WatchHistory.Add(WatchHistoryEntry);
+			}
+			Result = MLScreensReleaseWatchHistoryList(&WatchHistoryList);
+			UE_CLOG(Result != MLResult_Ok, LogScreensPlugin, Error, TEXT("MLScreensReleaseWatchHistoryList failed with error %s!"), UTF8_TO_TCHAR(MLScreensGetResultString(Result)));
+			Msg.bSuccess = true;
+		}
+		else
+		{
+			Msg.bSuccess = false;
+			UE_LOG(LogScreensPlugin, Error, TEXT("MLScreensGetWatchHistoryList failed with error %s!"), UTF8_TO_TCHAR(MLScreensGetResultString(Result)));
+		}
+	}
+#endif // WITH_MLSDK
+	return Msg;
+}
+
+bool FMagicLeapScreensPlugin::ClearWatchHistory()
+{
+#if WITH_MLSDK
+	{
+		FScopeLock Lock(&FMagicLeapScreensPlugin::CriticalSection);
+		MLScreensWatchHistoryList WatchHistoryList;
+		MLResult Result = MLScreensGetWatchHistoryList(&WatchHistoryList);
+		if (Result == MLResult_Ok)
+		{
+			for (uint32 i = 0; i < WatchHistoryList.count; ++i)
+			{
+				Result = MLScreensRemoveWatchHistoryEntry(WatchHistoryList.entries[i].id);
+				if (Result != MLResult_Ok)
+				{
+					UE_LOG(LogScreensPlugin, Error, TEXT("MLScreensRemoveWatchHistoryEntry failed with error %s for entry %d!"), UTF8_TO_TCHAR(MLScreensGetResultString(Result)), WatchHistoryList.entries[i].id);
+					return false;
+				}
+			}
+
+			Result = MLScreensReleaseWatchHistoryList(&WatchHistoryList);
+			if (Result != MLResult_Ok)
+			{
+				UE_LOG(LogScreensPlugin, Error, TEXT("MLScreensReleaseWatchHistoryList failed with error %s!"), UTF8_TO_TCHAR(MLScreensGetResultString(Result)));
+				return false;
+			}
+		}
+		else
+		{
+			UE_LOG(LogScreensPlugin, Error, TEXT("MLScreensGetWatchHistoryList failed with error %s!"), UTF8_TO_TCHAR(MLScreensGetResultString(Result)));
+			return false;
+		}
+	}
+#endif //WITH_MLSDK
+	return true;
+}
+
+FScreensMessage FMagicLeapScreensPlugin::AddToWatchHistory(const FScreensWatchHistoryEntry& WatchHistoryEntry)
+{
+	FScreensMessage Msg;
+	Msg.Type = EScreensMsgType::Response;
+	Msg.TaskType = EScreensTaskType::AddToHistory;
+
+#if WITH_MLSDK
+	{
+		FScopeLock Lock(&FMagicLeapScreensPlugin::CriticalSection);
+		MLScreensWatchHistoryEntry Entry;
+		Entry.title = TCHAR_TO_UTF8(*WatchHistoryEntry.Title);
+		Entry.subtitle = TCHAR_TO_UTF8(*WatchHistoryEntry.Subtitle);
+		Entry.playback_position_ms = static_cast<uint32>(WatchHistoryEntry.PlaybackPosition.GetTotalMilliseconds());
+		Entry.playback_duration_ms = static_cast<uint32>(WatchHistoryEntry.PlaybackDuration.GetTotalMilliseconds());
+		Entry.custom_data = TCHAR_TO_UTF8(*WatchHistoryEntry.CustomData);
+		MLImage MLThumbnail;
+
+		if (FMagicLeapScreensPlugin::ShouldUseDefaultThumbnail(WatchHistoryEntry, MLThumbnail))
+		{
+			MLThumbnail = FMagicLeapScreensPlugin::DefaultThumbnail;
+		}
+
+		MLResult Result = MLScreensInsertWatchHistoryEntry(&Entry, &MLThumbnail);
+		if (Result != MLResult_Ok)
+		{
+			UE_LOG(LogScreensPlugin, Error, TEXT("MLScreensInsertWatchHistoryEntry failed with error %s!"), UTF8_TO_TCHAR(MLScreensGetResultString(Result)));
+			Msg.WatchHistory.Add(WatchHistoryEntry);
+			Msg.bSuccess = false;
+		}
+		else
+		{
+			FScreensWatchHistoryEntry NewEntry;
+			FMagicLeapScreensPlugin::MLWatchHistoryEntryToUnreal(Entry, NewEntry);
+			Msg.WatchHistory.Add(NewEntry);
+			Msg.bSuccess = true;
+		}
+	}
+#endif // WITH_MLSDK
+	return Msg;
+}
+
+FScreensMessage FMagicLeapScreensPlugin::UpdateWatchHistoryEntry(const FScreensWatchHistoryEntry& WatchHistoryEntry)
+{
+	FScreensMessage Msg;
+	Msg.Type = EScreensMsgType::Response;
+	Msg.TaskType = EScreensTaskType::UpdateEntry;
+
+#if WITH_MLSDK
+	{
+		FScopeLock Lock(&FMagicLeapScreensPlugin::CriticalSection);
+		MLScreensWatchHistoryEntry Entry;
+		Entry.id = WatchHistoryEntry.ID.ID;
+		Entry.title = TCHAR_TO_UTF8(*WatchHistoryEntry.Title);
+		Entry.subtitle = TCHAR_TO_UTF8(*WatchHistoryEntry.Subtitle);
+		Entry.playback_position_ms = static_cast<uint32>(WatchHistoryEntry.PlaybackPosition.GetTotalMilliseconds());
+		Entry.playback_duration_ms = static_cast<uint32>(WatchHistoryEntry.PlaybackDuration.GetTotalMilliseconds());
+		Entry.custom_data = TCHAR_TO_UTF8(*WatchHistoryEntry.CustomData);
+		MLImage MLThumbnail;
+
+		if (FMagicLeapScreensPlugin::ShouldUseDefaultThumbnail(WatchHistoryEntry, MLThumbnail))
+		{
+			MLThumbnail = FMagicLeapScreensPlugin::DefaultThumbnail;
+		}
+
+		MLResult Result = MLScreensUpdateWatchHistoryEntry(&Entry, &MLThumbnail);
+		if (Result != MLResult_Ok)
+		{
+			UE_LOG(LogScreensPlugin, Error, TEXT("MLScreensUpdateWatchHistoryEntry failed with error %s!"), UTF8_TO_TCHAR(MLScreensGetResultString(Result)));
+			Msg.WatchHistory.Add(WatchHistoryEntry);
+			Msg.bSuccess = false;
+		}
+		else
+		{
+			FScreensWatchHistoryEntry UpdatedEntry;
+			MLWatchHistoryEntryToUnreal(Entry, UpdatedEntry);
+			Msg.WatchHistory.Add(UpdatedEntry);
+			Msg.bSuccess = true;
+		}
+	}
+#endif // WITH_MLSDK
+	return Msg;
+}
+
+bool FMagicLeapScreensPlugin::Tick(float DeltaTime)
+{
+	if (!FMagicLeapScreensPlugin::Impl->OutgoingMessages.IsEmpty())
+	{
+		FScreensMessage Msg;
+		FMagicLeapScreensPlugin::Impl->OutgoingMessages.Dequeue(Msg);
+
+		if (Msg.Type == EScreensMsgType::Request)
+		{
+			UE_LOG(LogScreensPlugin, Error, TEXT("Unexpected EScreensMsgType::Request received from worker thread!"));
+		}
+		else if (Msg.Type == EScreensMsgType::Response)
+		{
+			switch (Msg.TaskType)
+			{
+			case EScreensTaskType::None: break;
+			case EScreensTaskType::GetHistory:
+			{
+				Msg.HistoryDelegate.ExecuteIfBound(Msg.bSuccess, Msg.WatchHistory);
+			}
+			break;
+			case EScreensTaskType::AddToHistory:
+			{
+				Msg.EntryDelegate.ExecuteIfBound(Msg.bSuccess, Msg.WatchHistory[0]);
+			}
+			break;
+			case EScreensTaskType::UpdateEntry:
+			{
+				Msg.EntryDelegate.ExecuteIfBound(Msg.bSuccess, Msg.WatchHistory[0]);
+			}
+			break;
+			}
+		}
+	}
+	return true;
+}
+
+bool FMagicLeapScreensPlugin::GetScreensTransforms(TArray<FScreenTransform>& ScreensTransforms)
+{
+#if WITH_MLSDK
+	ScreensTransforms.Empty();
+
+	if (!IMagicLeapPlugin::Get().IsMagicLeapHMDValid())
+	{
+		return false;
+	}
+
+	const FAppFramework& AppFramework = static_cast<FMagicLeapHMD*>(GEngine->XRSystem->GetHMDDevice())->GetAppFrameworkConst();
+	if (!AppFramework.IsInitialized())
+	{
+		return false;
+	}
+	float WorldToMetersScale = AppFramework.GetWorldToMetersScale();
+
+	MLScreensScreenInfoListEx ScreensInfoList;
+	MLScreensScreenInfoListExInit(&ScreensInfoList);
+	MLResult Result = MLScreensGetScreenInfoListEx(&ScreensInfoList);
+	if (Result == MLResult_Ok)
+	{
+		// TODO: Param to GetTrackingToWorldTransform is unused (?)
+		FTransform PoseTransform = UHeadMountedDisplayFunctionLibrary::GetTrackingToWorldTransform(nullptr);
+		for (uint32 i = 0; i < ScreensInfoList.count; ++i)
+		{
+			MLScreensScreenInfoEx& Entry = ScreensInfoList.entries[i];
+			// TODO: Remove this once Screens team fixes setting dimensions when returning the ScreenInfoEx list
+			MLVec3f Scale = MagicLeap::ScaleFromMLMatrix(Entry.transform);
+			Entry.dimensions.x = 0.874f * Scale.x;
+			Entry.dimensions.y = 0.611f * Scale.y;
+			Entry.dimensions.z = 0.5f * Scale.z;
+			FScreenTransform ScreenTransform;
+
+			FTransform EntryTransform = FTransform(FQuat(FVector(0.0, 0.0, 1.0), PI) * MagicLeap::ToFQuat(Entry.transform), MagicLeap::ToFVector(Entry.transform, WorldToMetersScale), FVector(1.0, 1.0, 1.0));
+			if (EntryTransform.ContainsNaN())
+			{
+				UE_LOG(LogScreensPlugin, Error, TEXT("Screens info entry %d transform contains NaN."), Entry.screen_id);
+				continue;
+			}
+			if (!EntryTransform.GetRotation().IsNormalized())
+			{
+				FQuat rotation = EntryTransform.GetRotation();
+				rotation.Normalize();
+				EntryTransform.SetRotation(rotation);
+			}
+			EntryTransform.AddToTranslation(PoseTransform.GetLocation());
+			EntryTransform.ConcatenateRotation(PoseTransform.Rotator().Quaternion());
+
+			ScreenTransform.ScreenPosition = EntryTransform.GetLocation();
+			ScreenTransform.ScreenOrientation = EntryTransform.Rotator();
+			ScreenTransform.ScreenDimensions = MagicLeap::ToFVector(Entry.dimensions, WorldToMetersScale);
+			ScreenTransform.ScreenDimensions.X = FMath::Abs<float>(ScreenTransform.ScreenDimensions.X);
+			ScreenTransform.ScreenDimensions.Y = FMath::Abs<float>(ScreenTransform.ScreenDimensions.Y);
+			ScreenTransform.ScreenDimensions.Z = FMath::Abs<float>(ScreenTransform.ScreenDimensions.Z);
+
+			ScreensTransforms.Add(ScreenTransform);
+		}
+		Result = MLScreensReleaseScreenInfoListEx(&ScreensInfoList);
+		if (Result != MLResult_Ok)
+		{
+			UE_LOG(LogScreensPlugin, Error, TEXT("MLScreensReleaseScreenInfoListEx failed with error %s!"), UTF8_TO_TCHAR(MLScreensGetResultString(Result)));
+			return false;
+		}
+	}
+	else
+	{
+		UE_LOG(LogScreensPlugin, Error, TEXT("MLScreensGetScreenInfoListEx failed with error %s!"), UTF8_TO_TCHAR(MLScreensGetResultString(Result)));
+		return false;
+	}
+#endif // WITH_MLSDK
+	return true;
+}
+
+IMPLEMENT_MODULE(FMagicLeapScreensPlugin, MagicLeapScreens);
\ No newline at end of file
diff --git a/Engine/Plugins/Lumin/MagicLeapScreens/Source/MagicLeapScreens/Private/MagicLeapScreensPlugin.h b/Engine/Plugins/Lumin/MagicLeapScreens/Source/MagicLeapScreens/Private/MagicLeapScreensPlugin.h
new file mode 100644
--- /dev/null
+++ b/Engine/Plugins/Lumin/MagicLeapScreens/Source/MagicLeapScreens/Private/MagicLeapScreensPlugin.h
@@ -0,0 +1,60 @@
+// Copyright 1998-2018 Epic Games, Inc. All Rights Reserved.
+
+#pragma once
+
+#include "CoreMinimal.h"
+#include "Engine/Texture2D.h"
+#include "Containers/Array.h"
+#include "Misc/Optional.h"
+#include "IMagicLeapPlugin.h"
+#include "IMagicLeapScreensPlugin.h"
+#include "MagicLeapUtils.h"
+#include "MagicLeapScreensMsg.h"
+#include "MagicLeapPluginUtil.h" // for ML_INCLUDES_START/END
+
+#if WITH_MLSDK
+ML_INCLUDES_START
+#include <ml_image.h>
+#include <ml_screens.h>
+ML_INCLUDES_END
+#endif //WITH_MLSDK
+
+class FMagicLeapScreensPlugin : public IMagicLeapScreensPlugin
+{
+public:
+	void StartupModule() override;
+	void ShutdownModule() override;
+	bool IsEngineLoopInitComplete() const override;
+	void OnEngineLoopInitComplete() override;
+	static bool IsSupportedFormat(EPixelFormat InPixelFormat);
+
+#if WITH_MLSDK
+	static UTexture2D* MLImageToUTexture2D(const MLImage& Source);
+	static void MLWatchHistoryEntryToUnreal(const MLScreensWatchHistoryEntry& InEntry, FScreensWatchHistoryEntry& OutEntry);
+	static bool UTexture2DToMLImage(const UTexture2D& Source, MLImage& Target);
+#endif // WITH_MLSDK
+
+	static bool RemoveWatchHistoryEntry(const FScreenID& ID);
+	static FScreensMessage GetWatchHistoryEntries();
+	static bool ClearWatchHistory();
+	bool Tick(float DeltaTime);
+	static FScreensMessage AddToWatchHistory(const FScreensWatchHistoryEntry& WatchHistoryEntry);
+
+	static void GetWatchHistoryEntriesAsync(const TOptional<FScreensHistoryRequestResultDelegate>& OptionalResultDelegate);
+	static void AddToWatchHistoryAsync(const FScreensWatchHistoryEntry& NewEntry, const TOptional<FScreensEntryRequestResultDelegate>& OptionalResultDelegate);
+	static void UpdateWatchHistoryEntryAsync(const FScreensWatchHistoryEntry& UpdateEntry, const TOptional<FScreensEntryRequestResultDelegate>& OptionalResultDelegate);
+
+	static FScreensMessage UpdateWatchHistoryEntry(const FScreensWatchHistoryEntry& WatchHistoryEntry);
+	static bool GetScreensTransforms(TArray<FScreenTransform>& ScreensTransforms);
+	static TArray<uint8> PixelDataMemPool;
+	static class FScreensWorker *Impl;
+
+private:
+#if WITH_MLSDK
+	static bool ShouldUseDefaultThumbnail(const FScreensWatchHistoryEntry& Entry, MLImage& MLImage);
+	static MLImage DefaultThumbnail;
+	static FCriticalSection CriticalSection;
+#endif // WITH_MLSDK
+	FMagicLeapAPISetup APISetup;
+	bool bEngineLoopInitComplete;
+};
diff --git a/Engine/Plugins/Lumin/MagicLeapScreens/Source/MagicLeapScreens/Private/MagicLeapScreensWorker.h b/Engine/Plugins/Lumin/MagicLeapScreens/Source/MagicLeapScreens/Private/MagicLeapScreensWorker.h
new file mode 100644
--- /dev/null
+++ b/Engine/Plugins/Lumin/MagicLeapScreens/Source/MagicLeapScreens/Private/MagicLeapScreensWorker.h
@@ -0,0 +1,119 @@
+// Copyright 1998-2018 Epic Games, Inc. All Rights Reserved.
+
+#pragma once
+
+#include "MagicLeapScreensMsg.h"
+#include "Async/Async.h"
+#include "Containers/Queue.h"
+#include "MagicLeapHMD.h"
+#include "MagicLeapUtils.h"
+#include "MagicLeapScreensPlugin.h"
+
+class FScreensWorker : public FRunnable, public MagicLeap::IAppEventHandler
+{
+public:
+	FScreensWorker()
+		: Thread(nullptr)
+		, StopTaskCounter(0)
+		, Semaphore(nullptr)
+	{}
+
+	~FScreensWorker()
+	{
+		StopTaskCounter.Increment();
+		if (Semaphore != nullptr)
+		{
+			Semaphore->Trigger();
+			Thread->WaitForCompletion();
+			FGenericPlatformProcess::ReturnSynchEventToPool(Semaphore);
+			Semaphore = nullptr;
+			delete Thread;
+			Thread = nullptr;
+		}
+	}
+
+	void EngineInited()
+	{
+		if (Semaphore == nullptr)
+		{
+			Semaphore = FGenericPlatformProcess::GetSynchEventFromPool(false);
+#if PLATFORM_LUMIN
+			Thread = FRunnableThread::Create(this, TEXT("FScreensWorker"), 0, TPri_BelowNormal, FLuminAffinity::GetPoolThreadMask());
+#else
+			Thread = FRunnableThread::Create(this, TEXT("FScreensWorker"), 0, TPri_BelowNormal);
+#endif // PLATFORM_LUMIN
+		}
+		// wake up the worker to process the event
+		Semaphore->Trigger();
+	}
+
+	virtual uint32 Run() override
+	{
+		while (StopTaskCounter.GetValue() == 0)
+		{
+			if (IncomingMessages.Dequeue(CurrentMessage))
+			{
+				DoScreensTasks();
+			}
+			else
+			{
+				Semaphore->Wait();
+			}
+		}
+
+		return 0;
+	}
+
+	void ProcessMessage(const FScreensMessage& InMsg)
+	{
+		IncomingMessages.Enqueue(InMsg);
+		if (Semaphore != nullptr)
+		{
+			// wake up the worker to process the event
+			Semaphore->Trigger();
+		}
+	}
+
+	void DoScreensTasks()
+	{
+		switch (CurrentMessage.TaskType)
+		{
+		case EScreensTaskType::None: break;
+		case EScreensTaskType::GetHistory: GetWatchHistory(); break;
+		case EScreensTaskType::AddToHistory: AddToHistory();  break;
+		case EScreensTaskType::UpdateEntry: UpdateWatchHistoryEntry(); break;
+		}
+	}
+
+	void AddToHistory()
+	{
+		check(CurrentMessage.WatchHistory.Num() != 0);
+		FScreensMessage Msg = FMagicLeapScreensPlugin::AddToWatchHistory(CurrentMessage.WatchHistory[0]);
+		Msg.EntryDelegate = CurrentMessage.EntryDelegate;
+		check(Msg.WatchHistory.Num() > 0);
+		OutgoingMessages.Enqueue(Msg);
+	}
+
+	void UpdateWatchHistoryEntry()
+	{
+		check(CurrentMessage.WatchHistory.Num() != 0);
+		FScreensMessage Msg = FMagicLeapScreensPlugin::UpdateWatchHistoryEntry(CurrentMessage.WatchHistory[0]);
+		Msg.EntryDelegate = CurrentMessage.EntryDelegate;
+		check(Msg.WatchHistory.Num() > 0);
+		OutgoingMessages.Enqueue(Msg);
+	}
+
+	void GetWatchHistory()
+	{
+		FScreensMessage Msg = FMagicLeapScreensPlugin::GetWatchHistoryEntries();
+		Msg.HistoryDelegate = CurrentMessage.HistoryDelegate;
+		OutgoingMessages.Enqueue(Msg);
+	}
+
+	FRunnableThread* Thread;
+	FThreadSafeCounter StopTaskCounter;
+	TQueue<FScreensMessage, EQueueMode::Spsc> IncomingMessages;
+	TQueue<FScreensMessage, EQueueMode::Spsc> OutgoingMessages;
+	FEvent* Semaphore;
+	FScreensMessage CurrentMessage;
+};
diff --git a/Engine/Plugins/Lumin/MagicLeapScreens/Source/MagicLeapScreens/Public/IMagicLeapScreensPlugin.h b/Engine/Plugins/Lumin/MagicLeapScreens/Source/MagicLeapScreens/Public/IMagicLeapScreensPlugin.h
--- a/Engine/Plugins/Lumin/MagicLeapScreens/Source/MagicLeapScreens/Public/IMagicLeapScreensPlugin.h
+++ b/Engine/Plugins/Lumin/MagicLeapScreens/Source/MagicLeapScreens/Public/IMagicLeapScreensPlugin.h
@@ -3,6 +3,7 @@
 #pragma once
 
 #include "Modules/ModuleManager.h"
+#include "Containers/Ticker.h"
 
 /**
  * The public interface to this module.  In most cases, this interface is only public to sibling modules
@@ -33,7 +34,12 @@ public:
 		return FModuleManager::Get().IsModuleLoaded( "MagicLeapScreens" );
 	}
 
+	virtual bool Tick(float DeltaTime) = 0;
+
 	virtual bool IsEngineLoopInitComplete() const = 0;
 
 	virtual void OnEngineLoopInitComplete() = 0;
+
+	FTickerDelegate TickDelegate;
+	FDelegateHandle TickDelegateHandle;
 };
diff --git a/Engine/Plugins/Lumin/MagicLeapScreens/Source/MagicLeapScreens/Public/MagicLeapScreensComponent.h b/Engine/Plugins/Lumin/MagicLeapScreens/Source/MagicLeapScreens/Public/MagicLeapScreensComponent.h
new file mode 100644
--- /dev/null
+++ b/Engine/Plugins/Lumin/MagicLeapScreens/Source/MagicLeapScreens/Public/MagicLeapScreensComponent.h
@@ -0,0 +1,119 @@
+// Copyright 1998-2018 Epic Games, Inc. All Rights Reserved.
+
+#pragma once
+
+#include "Components/ActorComponent.h"
+#include "Engine/Engine.h"
+#include "Kismet/BlueprintFunctionLibrary.h"
+#include "MagicLeapScreensTypes.h"
+#include "MagicLeapScreensComponent.generated.h"
+
+/**
+	Component that provides access to screens functionality.
+*/
+UCLASS(ClassGroup = MagicLeap, BlueprintType, Blueprintable, EditInlineNew, meta = (BlueprintSpawnableComponent))
+class MAGICLEAPSCREENS_API UScreensComponent : public UActorComponent
+{
+	GENERATED_BODY()
+
+public:
+	UScreensComponent();
+	virtual ~UScreensComponent();
+
+	/** Polls for incoming messages from the worker thread */
+	virtual void TickComponent(float DeltaTime, enum ELevelTick TickType, FActorComponentTickFunction *ThisTickFunction) override;
+
+	/** 
+		Queues a task on the worker thread to retrieve the watch history. 
+		Subscribe to the GetWatchHistorySuccess & GetWatchHistoryFailure delegates to receive the result.
+	*/
+	UFUNCTION(BlueprintCallable, Category = "Screens|MagicLeap", meta = (DeprecatedFunction, DeprecationMessage = "Please use the Magic Leap Screens Function Library instead."))
+	void GetWatchHistoryAsync();
+
+	/** 
+		Queues a task on the worker thread to add a new entry into the watch history.
+		Subscribe to the AddToWatchHistoryResult delegate to receive the result.
+	
+		@param[in] WatchHistoryEntry The entry to be added to the history.
+	*/
+	UFUNCTION(BlueprintCallable, Category = "Screens|MagicLeap", meta = (DeprecatedFunction, DeprecationMessage = "Please use the Magic Leap Screens Function Library instead."))
+	void AddWatchHistoryEntryAsync(const FScreensWatchHistoryEntry& WatchHistoryEntry);
+
+	/**
+		Queues a task on the worker thread to update an entry in the watch history.
+		The ID of the entry to be updated must be valid and already present in the watch history.
+		Subscribe to the UpdateWatchHistoryEntryResult delegate to receive the result.
+	
+		@param[in] WatchHistoryEntry The entry to be updated.
+	*/
+	UFUNCTION(BlueprintCallable, Category = "Screens|MagicLeap", meta = (DeprecatedFunction, DeprecationMessage = "Please use the Magic Leap Screens Function Library instead."))
+	void UpdateWatchHistoryEntryAsync(const FScreensWatchHistoryEntry& WatchHistoryEntry);
+
+	/**
+		Removes a watch history entry.
+		
+		@param[in] FScreenID The id of the entry to find and remove.
+		@return True if no errors were encountered.
+	*/
+	UFUNCTION(BlueprintCallable, Category = "Screens|MagicLeap", meta = (DeprecatedFunction, DeprecationMessage = "Please use the Magic Leap Screens Function Library instead."))
+	bool RemoveWatchHistoryEntry(const FScreenID& ID);
+
+	/**
+		Permanently clears all watch history entries from disk.
+	
+		@return True if no errors were encountered.
+	*/
+	UFUNCTION(BlueprintCallable, Category = "Screens|MagicLeap", meta = (DeprecatedFunction, DeprecationMessage = "Please use the Magic Leap Screens Function Library instead."))
+	bool ClearWatchHistory();
+
+	/** 
+		Retrieves a list of screen transforms.
+		
+		@param[out] ScreensTransforms A list of transforms to be provided by the underlying api.
+		@return True if no errors were encountered.
+	*/
+	UFUNCTION(BlueprintCallable, Category = "Screens|MagicLeap", meta = (DeprecatedFunction, DeprecationMessage = "Please use the Magic Leap Screens Function Library instead."))
+	bool GetScreensTransforms(TArray<FScreenTransform>& ScreensTransforms);
+
+public:
+	/** 
+		Delegate used to notify the instigating blueprint of a watch history retrieval success.
+
+		@param[out] WatchHistory A list of history entries requested via a call to GetWatchHistoryAsync().
+	*/
+	DECLARE_DYNAMIC_MULTICAST_DELEGATE_OneParam(FScreensGetWatchHistorySuccess, const TArray<FScreensWatchHistoryEntry>&, WatchHistory);
+
+	/** Delegate used to notify the instigating blueprint of a watch history retrieval failure. */
+	DECLARE_DYNAMIC_MULTICAST_DELEGATE(FScreensGetWatchHistoryFailure);
+
+	/** 
+		Delegate used to notify the instigating blueprint of the result from a request to add a new watch history entry.
+
+		@param[out] Entry The new entry created if a request via a call to AddToWatchHistoryAsync() succeeds.
+		@param[out] Success True when the call to AddToWatchHistoryAsync() succeeds.
+	*/
+	DECLARE_DYNAMIC_MULTICAST_DELEGATE_TwoParams(FScreensAddToWatchHistoryResult, const FScreensWatchHistoryEntry&, Entry, const bool, Success);
+
+	/**
+		Delegate used to notify the instigating blueprint of the result from a request to update an entry in the watch history
+
+		@param[out] Entry The updated entry resulting from a successful request via a call to UpdateWatchHistoryEntry().
+		@param[out] Success True when the call to UpdateWatchHistoryEntryAsync() succeeds.
+	*/
+	DECLARE_DYNAMIC_MULTICAST_DELEGATE_TwoParams(FScreensUpdateWatchHistoryEntryResult, const FScreensWatchHistoryEntry&, Entry, const bool, Success);
+
+private:
+
+	UPROPERTY(BlueprintAssignable, Category = "Screens|MagicLeap", meta = (AllowPrivateAccess = true))
+	FScreensGetWatchHistorySuccess GetWatchHistorySuccess;
+
+	UPROPERTY(BlueprintAssignable, Category = "Screens|MagicLeap", meta = (AllowPrivateAccess = true))
+	FScreensGetWatchHistoryFailure GetWatchHistoryFailure;
+
+	UPROPERTY(BlueprintAssignable, Category = "Screens|MagicLeap", meta = (AllowPrivateAccess = true))
+	FScreensAddToWatchHistoryResult AddToWatchHistoryResult;
+
+	UPROPERTY(BlueprintAssignable, Category = "Screens|MagicLeap", meta = (AllowPrivateAccess = true))
+	FScreensUpdateWatchHistoryEntryResult UpdateWatchHistoryEntryResult;
+};
+
diff --git a/Engine/Plugins/Lumin/MagicLeapScreens/Source/MagicLeapScreens/Public/MagicLeapScreensFunctionLibrary.h b/Engine/Plugins/Lumin/MagicLeapScreens/Source/MagicLeapScreens/Public/MagicLeapScreensFunctionLibrary.h
new file mode 100644
--- /dev/null
+++ b/Engine/Plugins/Lumin/MagicLeapScreens/Source/MagicLeapScreens/Public/MagicLeapScreensFunctionLibrary.h
@@ -0,0 +1,81 @@
+// Copyright 1998-2018 Epic Games, Inc. All Rights Reserved.
+
+#pragma once
+
+#include "Engine/Engine.h"
+#include "Kismet/BlueprintFunctionLibrary.h"
+#include "MagicLeapScreensTypes.h"
+#include "MagicLeapScreensFunctionLibrary.generated.h"
+
+UCLASS(ClassGroup = MagicLeap)
+class MAGICLEAPSCREENS_API UMagicLeapScreensFunctionLibrary : public UBlueprintFunctionLibrary
+{
+	GENERATED_BODY()
+
+public:
+
+	/**
+		Asynchronously requests all watch history entries.
+
+		@param[in] ResultDelegate Delegate that is called once the request has been resolved.
+	*/
+	UFUNCTION(BlueprintCallable, Category = "Screens Function Library|MagicLeap")
+	static void GetWatchHistoryAsync(const FScreensHistoryRequestResultDelegate& ResultDelegate);
+
+	/**
+		Asynchronously requests to adds a new entry into the watch history.
+
+		@param[in] NewEntry Watch history entry to add.
+		@param[in] ResultDelegate Delegate that is called once the add request has been resolved. The resulting FScreensWatchHistoryEntry returned.
+		on success contains the new ID of the added entry.
+	*/
+	UFUNCTION(BlueprintCallable, Category = "Screens Function Library|MagicLeap")
+	static void AddToWatchHistoryAsync(const FScreensWatchHistoryEntry& NewEntry, const FScreensEntryRequestResultDelegate& ResultDelegate);
+
+	/**
+		Asynchronously requests to update an entry in the watch history
+
+		@param[in] UpdateEntry Watch history entry to update. The ID of this entry must be valid and in the watch history in order for the update to resolve successfully
+		@param[in] ResultDelegate Delegate that is called once the update request has been resolved.
+	*/
+	UFUNCTION(BlueprintCallable, Category = "Screens Function Library|MagicLeap")
+	static void UpdateWatchHistoryEntryAsync(const FScreensWatchHistoryEntry& UpdateEntry, const FScreensEntryRequestResultDelegate& ResultDelegate);
+
+	/**
+		Removes an entry from the watch history that corresponse with the ID passed.
+
+		@param[in] ID ID of the watch history entry to remove.
+		@return True when the entry is sucessfully removed.
+	*/
+	UFUNCTION(BlueprintCallable, Category = "Screens Function Library|MagicLeap")
+	static bool RemoveWatchHistoryEntry(const FScreenID& ID);
+
+	/**
+		Removes all watch history entries.
+
+		@return True when all watch history entries have been successfully removed.
+	*/
+	UFUNCTION(BlueprintCallable, Category = "Screens Function Library|MagicLeap")
+	static bool ClearWatchHistory();
+
+	/**
+		Gets the transform for all watch history entries.
+
+		@param[out] ScreenTransforms Array of transforms that corresponse to all current watch history entries.
+		@return True when the request for all transforms succeeds 
+	*/
+	UFUNCTION(BlueprintCallable, Category = "Screens Function Library|MagicLeap")
+	static bool GetScreensTransforms(TArray<FScreenTransform>& ScreenTransforms);
+
+	/**
+		Delegate used to relay the results of GetWatchHistory.
+	*/
+	UPROPERTY()
+	FScreensHistoryRequestResultDelegate HistoryResultDelegate;
+
+	/**
+		Delegate used to relay the results of both AddToWatchHistory and UpdateWatchHistoryEntry.
+	*/
+	UPROPERTY()
+	FScreensEntryRequestResultDelegate EntryResultDelegate;
+};
diff --git a/Engine/Plugins/Lumin/MagicLeapScreens/Source/MagicLeapScreens/Public/MagicLeapScreensTypes.h b/Engine/Plugins/Lumin/MagicLeapScreens/Source/MagicLeapScreens/Public/MagicLeapScreensTypes.h
new file mode 100644
--- /dev/null
+++ b/Engine/Plugins/Lumin/MagicLeapScreens/Source/MagicLeapScreens/Public/MagicLeapScreensTypes.h
@@ -0,0 +1,97 @@
+// Copyright 1998-2018 Epic Games, Inc. All Rights Reserved.
+
+#pragma once
+
+#include "Engine/Engine.h"
+#include "MagicLeapScreensTypes.generated.h"
+
+/** 
+  ID for a Screens Watch History Entry.
+
+  Save this off when you add a new watch history and use the same to update or delete that same entry.
+ */
+USTRUCT(BlueprintType)
+struct FScreenID
+{
+	GENERATED_BODY()
+
+public:
+	int64 ID;
+};
+
+/** Channel watch history, may be displayed in the Screens Launcher application. */
+USTRUCT(BlueprintType)
+struct FScreensWatchHistoryEntry
+{
+	GENERATED_BODY()
+
+public:
+	/** Entry Identifier. Must be used to update and delete a given entry. */
+	UPROPERTY(VisibleAnywhere, BlueprintReadOnly, Category = "Screens|MagicLeap")
+	FScreenID ID;
+
+	/** Title of the media for which this entry is created. */
+	UPROPERTY(EditAnywhere, BlueprintReadWrite, Category = "Screens|MagicLeap")
+	FString Title;
+
+	/** Subtitle of the media for which this entry is created. */
+	UPROPERTY(EditAnywhere, BlueprintReadWrite, Category = "Screens|MagicLeap")
+	FString Subtitle;
+
+	/** Current media playback position. Can be fed from UMediaPlayer::GetTime(). */
+	UPROPERTY(EditAnywhere, BlueprintReadWrite, Category = "Screens|MagicLeap")
+	FTimespan PlaybackPosition;
+
+	/** Total duration of the media. Can be fed from UMediaPlayer::GetDuration() */
+	UPROPERTY(EditAnywhere, BlueprintReadWrite, Category = "Screens|MagicLeap")
+	FTimespan PlaybackDuration;
+
+	/** Any data the application might want to save off in the watch history and then receive back from the Screens Launher. */
+	UPROPERTY(EditAnywhere, BlueprintReadWrite, Category = "Screens|MagicLeap")
+	FString CustomData;
+
+	/** Thumbnail to be shown in the Screens Launcher application for this watch history entry. */
+	UPROPERTY(EditAnywhere, BlueprintReadWrite, Category = "Screens|MagicLeap")
+	UTexture2D* Thumbnail;
+};
+
+/**
+  Information required to place a screen in the world.
+
+  This will be received from the Screens Launcher api, based on the previous screens spawned by user.
+ */
+USTRUCT(BlueprintType)
+struct FScreenTransform
+{
+	GENERATED_BODY()
+
+public:
+	/** Position of the screen in Unreal's world space. */
+	UPROPERTY(EditAnywhere, BlueprintReadWrite, Category = "Screens|MagicLeap")
+	FVector ScreenPosition;
+
+	/** Orientation of the screen in Unreal's world space. */
+	UPROPERTY(EditAnywhere, BlueprintReadWrite, Category = "Screens|MagicLeap")
+	FRotator ScreenOrientation;
+
+	/** Dimensions of the screen in Unreal Units. The dimensions are axis-aligned with the orientation. */
+	UPROPERTY(EditAnywhere, BlueprintReadWrite, Category = "Screens|MagicLeap")
+	FVector ScreenDimensions;
+};
+
+/**
+	Delegate used to relay the result of a Screens operation that involves a single watch history entry.
+	For example updating or adding a history entry.
+
+	@param[out] bSuccess True when the request is successful
+	@param[out] WatchHistoryEntry Resulting watch history entry for which the operation was performed on.
+*/
+DECLARE_DYNAMIC_DELEGATE_TwoParams(FScreensEntryRequestResultDelegate, const bool, bSuccess, const FScreensWatchHistoryEntry&, WatchHistoryEntry);
+
+/**
+	Delegate used to relay the result of getting the entire watch history.
+
+	@param[out] bSuccess True when the request is successful.
+	@param[out] WatchHistoryEntry Resulting array of watch histories returned by the operation.
+*/
+DECLARE_DYNAMIC_DELEGATE_TwoParams(FScreensHistoryRequestResultDelegate, const bool, bSuccess, const TArray<FScreensWatchHistoryEntry>&, WatchHistoryEntries);
\ No newline at end of file
diff --git a/Engine/Plugins/Media/AvfMedia/Source/AvfMedia/Private/Player/AvfMediaPlayer.cpp b/Engine/Plugins/Media/AvfMedia/Source/AvfMedia/Private/Player/AvfMediaPlayer.cpp
--- a/Engine/Plugins/Media/AvfMedia/Source/AvfMedia/Private/Player/AvfMediaPlayer.cpp
+++ b/Engine/Plugins/Media/AvfMedia/Source/AvfMedia/Private/Player/AvfMediaPlayer.cpp
@@ -508,7 +508,7 @@ bool FAvfMediaPlayer::Open(const TSharedRef<FArchive, ESPMode::ThreadSafe>& /*Ar
 }
 
 
-void FAvfMediaPlayer::TickAudio()
+void FAvfMediaPlayer::TickAudio(FTimespan LastAudioSampleTime)
 {
 	// NOP
 }
diff --git a/Engine/Plugins/Media/AvfMedia/Source/AvfMedia/Private/Player/AvfMediaPlayer.h b/Engine/Plugins/Media/AvfMedia/Source/AvfMedia/Private/Player/AvfMediaPlayer.h
--- a/Engine/Plugins/Media/AvfMedia/Source/AvfMedia/Private/Player/AvfMediaPlayer.h
+++ b/Engine/Plugins/Media/AvfMedia/Source/AvfMedia/Private/Player/AvfMediaPlayer.h
@@ -64,7 +64,7 @@ public:
 	virtual IMediaView& GetView() override;
 	virtual bool Open(const FString& Url, const IMediaOptions* Options) override;
 	virtual bool Open(const TSharedRef<FArchive, ESPMode::ThreadSafe>& Archive, const FString& OriginalUrl, const IMediaOptions* Options) override;
-	virtual void TickAudio() override;
+	virtual void TickAudio(FTimespan LastAudioSampleTime) override;
 	virtual void TickFetch(FTimespan DeltaTime, FTimespan Timecode) override;
 	virtual void TickInput(FTimespan DeltaTime, FTimespan Timecode) override;
 
diff --git a/Engine/Plugins/Media/MfMedia/Source/MfMedia/Private/Player/MfMediaPlayer.cpp b/Engine/Plugins/Media/MfMedia/Source/MfMedia/Private/Player/MfMediaPlayer.cpp
--- a/Engine/Plugins/Media/MfMedia/Source/MfMedia/Private/Player/MfMediaPlayer.cpp
+++ b/Engine/Plugins/Media/MfMedia/Source/MfMedia/Private/Player/MfMediaPlayer.cpp
@@ -193,7 +193,7 @@ bool FMfMediaPlayer::Open(const TSharedRef<FArchive, ESPMode::ThreadSafe>& Archi
 }
 
 
-void FMfMediaPlayer::TickAudio()
+void FMfMediaPlayer::TickAudio(FTimespan LastAudioSampleTime)
 {
 	if (CurrentState == EMediaState::Playing)
 	{
diff --git a/Engine/Plugins/Media/MfMedia/Source/MfMedia/Private/Player/MfMediaPlayer.h b/Engine/Plugins/Media/MfMedia/Source/MfMedia/Private/Player/MfMediaPlayer.h
--- a/Engine/Plugins/Media/MfMedia/Source/MfMedia/Private/Player/MfMediaPlayer.h
+++ b/Engine/Plugins/Media/MfMedia/Source/MfMedia/Private/Player/MfMediaPlayer.h
@@ -69,7 +69,7 @@ public:
 	virtual bool Open(const FString& Url, const IMediaOptions* Options) override;
 	virtual bool Open(const FString& Url, const IMediaOptions* Options, const FMediaPlayerOptions* PlayerOptions) override;
 	virtual bool Open(const TSharedRef<FArchive, ESPMode::ThreadSafe>& Archive, const FString& OriginalUrl, const IMediaOptions* Options) override;
-	virtual void TickAudio() override;
+	virtual void TickAudio(FTimespan LastAudioSampleTime) override;
 	virtual void TickFetch(FTimespan DeltaTime, FTimespan Timecode) override;
 	virtual void TickInput(FTimespan DeltaTime, FTimespan Timecode) override;
 
diff --git a/Engine/Plugins/Runtime/AndroidDeviceProfileSelector/Source/AndroidDeviceProfileSelectorRuntime/Private/AndroidDeviceProfileSelectorRuntimeModule.cpp b/Engine/Plugins/Runtime/AndroidDeviceProfileSelector/Source/AndroidDeviceProfileSelectorRuntime/Private/AndroidDeviceProfileSelectorRuntimeModule.cpp
--- a/Engine/Plugins/Runtime/AndroidDeviceProfileSelector/Source/AndroidDeviceProfileSelectorRuntime/Private/AndroidDeviceProfileSelectorRuntimeModule.cpp
+++ b/Engine/Plugins/Runtime/AndroidDeviceProfileSelector/Source/AndroidDeviceProfileSelectorRuntime/Private/AndroidDeviceProfileSelectorRuntimeModule.cpp
@@ -21,13 +21,23 @@ void FAndroidDeviceProfileSelectorRuntimeModule::ShutdownModule()
 
 FString const FAndroidDeviceProfileSelectorRuntimeModule::GetRuntimeDeviceProfileName()
 {
+	static FString ProfileName;
+
 #if PLATFORM_LUMIN
+	if (ProfileName.IsEmpty())
+	{
+		// Fallback profiles in case we do not match any rules
+		ProfileName = FPlatformMisc::GetDefaultDeviceProfileName();
+		if (ProfileName.IsEmpty())
+		{
+			ProfileName = FPlatformProperties::PlatformName();
+		}
+	}
+
 	// @todo Lumin: when removing this, also remove Lumin from the .uplugin
-	UE_LOG(LogAndroid, Log, TEXT("Selected Device Profile: [Lumin]"));
-	return TEXT("Lumin");
+	UE_LOG(LogAndroid, Log, TEXT("Selected Device Profile: [%s]"), *ProfileName);
+	return ProfileName;
 #endif
-
-	static FString ProfileName;
 	
 	if (ProfileName.IsEmpty())
 	{
diff --git a/Engine/Source/Developer/Android/AndroidDeviceDetection/Private/AndroidDeviceDetectionModule.cpp b/Engine/Source/Developer/Android/AndroidDeviceDetection/Private/AndroidDeviceDetectionModule.cpp
--- a/Engine/Source/Developer/Android/AndroidDeviceDetection/Private/AndroidDeviceDetectionModule.cpp
+++ b/Engine/Source/Developer/Android/AndroidDeviceDetection/Private/AndroidDeviceDetectionModule.cpp
@@ -310,8 +310,9 @@ private:
 			}
 			else
 			{
-				// grab the Android version
-				const FString AndroidVersionCommand = FString::Printf(TEXT("-s %s %s ro.build.version.release"), *NewDeviceInfo.SerialNumber, *GetPropCommand);
+				// grab the Lumin/Android version
+				const FString AndroidVersionCommand = bForLumin ? FString::Printf(TEXT("%s ro.build.id"), *GetPropCommand) :
+					FString::Printf(TEXT("-s %s %s ro.build.version.release"), *NewDeviceInfo.SerialNumber, *GetPropCommand);
 				if (!ExecuteAdbCommand(*AndroidVersionCommand, &NewDeviceInfo.HumanAndroidVersion, nullptr))
 				{
 					continue;
diff --git a/Engine/Source/Developer/Lumin/LuminPlatformEditor/LuminPlatformEditor.Build.cs b/Engine/Source/Developer/Lumin/LuminPlatformEditor/LuminPlatformEditor.Build.cs
--- a/Engine/Source/Developer/Lumin/LuminPlatformEditor/LuminPlatformEditor.Build.cs
+++ b/Engine/Source/Developer/Lumin/LuminPlatformEditor/LuminPlatformEditor.Build.cs
@@ -27,8 +27,10 @@ public class LuminPlatformEditor : ModuleRules
 				"TargetPlatform",
 				"RenderCore",
 				"MaterialShaderQualitySettings",
-                "AudioSettingsEditor",
-            }
+				"AudioSettingsEditor",
+				"UnrealEd",
+				"PropertyPath",
+			}
 		);
 
 		PrivateIncludePathModuleNames.AddRange(
diff --git a/Engine/Source/Developer/Lumin/LuminPlatformEditor/Private/LuminPlatformEditorModule.cpp b/Engine/Source/Developer/Lumin/LuminPlatformEditor/Private/LuminPlatformEditorModule.cpp
--- a/Engine/Source/Developer/Lumin/LuminPlatformEditor/Private/LuminPlatformEditorModule.cpp
+++ b/Engine/Source/Developer/Lumin/LuminPlatformEditor/Private/LuminPlatformEditorModule.cpp
@@ -69,7 +69,15 @@ class FLuminPlatformEditorModule
 				);
 			}
 
-			// @todo Shader quality settings
+			{
+				static FName NAME_SF_VULKAN_ES31_LUMIN(TEXT("SF_VULKAN_ES31_LUMIN"));
+				UShaderPlatformQualitySettings* LuminMaterialQualitySettings = UMaterialShaderQualitySettings::Get()->GetShaderPlatformQualitySettings(NAME_SF_VULKAN_ES31_LUMIN);
+				SettingsModule->RegisterSettings("Project", "Platforms", "MagicLeapVulkanQuality",
+					LOCTEXT("LuminVulkanQualitySettingsName", "Lumin Material Quality - Vulkan"),
+					LOCTEXT("LuminVulkanQualitySettingsDescription", "Settings for Lumin Vulkan material quality."),
+					LuminMaterialQualitySettings
+				);
+			}
 		}
 
 		// Force the SDK settings into a sane state initially so we can make use of them
@@ -89,6 +97,7 @@ class FLuminPlatformEditorModule
 		{
 			SettingsModule->UnregisterSettings("Project", "Platforms", "Lumin");
 			SettingsModule->UnregisterSettings("Project", "Platforms", "MagicLeapSDK");
+			SettingsModule->UnregisterSettings("Project", "Platforms", "MagicLeapVulkanQuality");
 		}
 	}
 };
diff --git a/Engine/Source/Developer/Lumin/LuminPlatformEditor/Private/LuminTargetSettingsDetails.cpp b/Engine/Source/Developer/Lumin/LuminPlatformEditor/Private/LuminTargetSettingsDetails.cpp
--- a/Engine/Source/Developer/Lumin/LuminPlatformEditor/Private/LuminTargetSettingsDetails.cpp
+++ b/Engine/Source/Developer/Lumin/LuminPlatformEditor/Private/LuminTargetSettingsDetails.cpp
@@ -8,6 +8,20 @@
 #include "IDetailChildrenBuilder.h"
 #include "UObject/UnrealType.h"
 #include "Widgets/Text/STextBlock.h"
+#include "Widgets/Images/SImage.h"
+#include "Widgets/Input/SButton.h"
+#include "Widgets/Notifications/SNotificationList.h"
+#include "SPlatformSetupMessage.h"
+#include "DetailCategoryBuilder.h"
+#include "DesktopPlatformModule.h"
+#include "Framework/Application/SlateApplication.h"
+#include "Framework/Notifications/NotificationManager.h"
+#include "SourceControlHelpers.h"
+#include "LuminRuntimeSettings.h"
+#include "Logging/LogMacros.h"
+#include "PropertyPathHelpers.h"
+
+DEFINE_LOG_CATEGORY_STATIC(LogLuminTargetSettingsDetail, Log, All);
 
 #define LOCTEXT_NAMESPACE "LuminTargetSettingsDetails"
 
@@ -16,9 +30,426 @@ TSharedRef<IDetailCustomization> FLuminTargetSettingsDetails::MakeInstance()
 	return MakeShareable(new FLuminTargetSettingsDetails);
 }
 
+FLuminTargetSettingsDetails::FLuminTargetSettingsDetails()
+	: DefaultIconModelPath(FPaths::EngineDir() / TEXT("Build/Lumin/Resources/Model"))
+	, DefaultIconPortalPath(FPaths::EngineDir() / TEXT("Build/Lumin/Resources/Portal"))
+	, GameLuminPath(FPaths::ProjectDir() / TEXT("Build/Lumin"))
+	, GameProjectSetupPath(GameLuminPath / TEXT("IconSetup.txt"))
+	, SavedLayoutBuilder(nullptr)
+{
+}
+
 void FLuminTargetSettingsDetails::CustomizeDetails(IDetailLayoutBuilder& DetailBuilder)
+{
+	SavedLayoutBuilder = &DetailBuilder;
+
+	IconModelPathProp = DetailBuilder.GetProperty("IconModelPath.Path");
+	IconPortalPathProp = DetailBuilder.GetProperty("IconPortalPath.Path");
+	CertificateProp = DetailBuilder.GetProperty("Certificate.FilePath");
+
+	IconModelPathAttribute.Bind(TAttribute<FString>::FGetter::CreateRaw(this, &FLuminTargetSettingsDetails::IconModelPathGetter));
+	IconPortalPathAttribute.Bind(TAttribute<FString>::FGetter::CreateRaw(this, &FLuminTargetSettingsDetails::IconPortalPathGetter));
+	CertificateAttribute.Bind(TAttribute<FString>::FGetter::CreateRaw(this, &FLuminTargetSettingsDetails::CertificateGetter));
+
+	BuildAudioSection(DetailBuilder);
+	BuildAppTileSection(DetailBuilder);
+}
+
+FString FLuminTargetSettingsDetails::IconModelPathGetter()
+{
+	FString Result;
+	if (IconModelPathProp->GetValue(Result) == FPropertyAccess::Success && !Result.IsEmpty())
+	{
+		if (SetupForPlatformAttribute.Get())
+		{
+			// If we did the setup, the values are game/project dir relative. So we root them at that
+			// for display. Otherwise we show the default values which will show engine exec relative
+			// paths to the platform build resource locations.
+			Result = FPaths::ProjectDir() / Result;
+		}
+		else
+		{
+			// Otherwise they are in the engine tree. But they could be outdate paths.
+			// In which case we use the hard-wired defaults.
+			if (FPaths::DirectoryExists(FPaths::EngineDir() / Result))
+			{
+				Result = FPaths::EngineDir() / Result;
+			}
+			else
+			{
+				Result = DefaultIconModelPath;
+			}
+		}
+	}
+	return Result;
+}
+
+FString FLuminTargetSettingsDetails::IconPortalPathGetter()
+{
+	FString Result;
+	if (IconPortalPathProp->GetValue(Result) == FPropertyAccess::Success && !Result.IsEmpty())
+	{
+		if (SetupForPlatformAttribute.Get())
+		{
+			// If we did the setup, the values are game/project dir relative. So we root them at that
+			// for display. Otherwise we show the default values which will show engine exec relative
+			// paths to the platform build resource locations.
+			Result = FPaths::ProjectDir() / Result;
+		}
+		else
+		{
+			// Otherwise they are in the engine tree. But they could be outdate paths.
+			// In which case we use the hard-wired defaults.
+			if (FPaths::DirectoryExists(FPaths::EngineDir() / Result))
+			{
+				Result = FPaths::EngineDir() / Result;
+			}
+			else
+			{
+				Result = DefaultIconPortalPath;
+			}
+		}
+	}
+	return Result;
+}
+
+FString FLuminTargetSettingsDetails::CertificateGetter()
+{
+	FString Result;
+	if (CertificateProp->GetValue(Result) == FPropertyAccess::Success && !Result.IsEmpty())
+	{
+		Result = FPaths::ProjectDir() / Result;
+	}
+	return Result;
+}
+
+void FLuminTargetSettingsDetails::CopySetupFilesIntoProject()
+{
+	// Start out with the hard-wired default paths.
+	FString SourceModelPath = DefaultIconModelPath;
+	FString SourcePortalPath = DefaultIconPortalPath;
+	// Override with soft-wired defaults from the engine config. These are going to be engine
+	// root relative as we are copying from the engine tree to the project tree.
+	// But only override if those soft-wired engine paths exists. If they don't it's likely
+	// some old invalid value. This prevents someone from using obsolete data.
+	FString PropSourceModelPath;
+	FString PropSourcePortalPath;
+	if (IconModelPathProp->GetValue(PropSourceModelPath) == FPropertyAccess::Success && !PropSourceModelPath.IsEmpty() && FPaths::DirectoryExists(FPaths::EngineDir() / PropSourceModelPath))
+	{
+		SourceModelPath = FPaths::EngineDir() / PropSourceModelPath;
+	}
+	if (IconPortalPathProp->GetValue(PropSourcePortalPath) == FPropertyAccess::Success && !PropSourcePortalPath.IsEmpty() && FPaths::DirectoryExists(FPaths::EngineDir() / PropSourcePortalPath))
+	{
+		SourcePortalPath = FPaths::EngineDir() / PropSourcePortalPath;
+	}
+	const FString TargetModelPath = GameLuminPath / TEXT("Model");
+	const FString TargetPortalPath = GameLuminPath / TEXT("Portal");
+	bool DidModelCopy = CopyDir(*SourceModelPath, *TargetModelPath);
+	bool DidPortalCopy = CopyDir(*SourcePortalPath, *TargetPortalPath);
+	if (DidModelCopy && DidPortalCopy)
+	{
+		// Touch the setup file to indicate we did the copies.
+		delete IPlatformFile::GetPlatformPhysical().OpenWrite(*GameProjectSetupPath);
+		// And set the icon path config vars to the project directory now that we have it.
+		// This makes it so that the packaging will use these instead of the engine
+		// files directly. The values for both are fixed to the project root relative locations.
+		if (IconModelPathProp->SetValue(FString("Build/Lumin/Model")) != FPropertyAccess::Success ||
+			IconPortalPathProp->SetValue(FString("Build/Lumin/Portal")) != FPropertyAccess::Success)
+		{
+			UE_LOG(LogLuminTargetSettingsDetail, Error, TEXT("Failed to update icon or portal "));
+		}
+	}
+}
+
+void FLuminTargetSettingsDetails::BuildAudioSection(IDetailLayoutBuilder& DetailBuilder)
 {
 	AudioPluginManager.BuildAudioCategory(DetailBuilder, EAudioPlatform::Lumin);
 }
 
+void FLuminTargetSettingsDetails::BuildAppTileSection(IDetailLayoutBuilder& DetailBuilder)
+{
+	////////// UI for icons..
+
+	IDetailCategoryBuilder& AppTitleCategory = DetailBuilder.EditCategory(TEXT("Magic Leap App Tile"));
+	DetailBuilder.HideProperty("IconModelPath");
+	DetailBuilder.HideProperty("IconPortalPath");
+	TSharedRef<SPlatformSetupMessage> PlatformSetupMessage = SNew(SPlatformSetupMessage, GameProjectSetupPath)
+		.PlatformName(LOCTEXT("LuminPlatformName", "Magic Leap"))
+		.OnSetupClicked(FSimpleDelegate::CreateLambda([this] { this->CopySetupFilesIntoProject(); }));
+	SetupForPlatformAttribute = PlatformSetupMessage->GetReadyToGoAttribute();
+	AppTitleCategory.AddCustomRow(LOCTEXT("Warning", "Warning"), false)
+		.WholeRowWidget
+		[
+			PlatformSetupMessage
+		];
+	AppTitleCategory.AddCustomRow(LOCTEXT("BuildFolderLabel", "Build Folder"), false)
+		.IsEnabled(SetupForPlatformAttribute)
+		.NameContent()
+		[
+			SNew(SHorizontalBox)
+			+ SHorizontalBox::Slot()
+			.Padding(FMargin(0, 1, 0, 1))
+			.FillWidth(1.0f)
+			[
+				SNew(STextBlock)
+				.Text(LOCTEXT("BuildFolderLabel", "Build Folder"))
+				.Font(DetailBuilder.GetDetailFont())
+			]
+		]
+		.ValueContent()
+		[
+			SNew(SHorizontalBox)
+			+ SHorizontalBox::Slot()
+			.AutoWidth()
+			[
+				SNew(SButton)
+				.Text(LOCTEXT("OpenBuildFolderButton", "Open Build Folder"))
+				.ToolTipText(LOCTEXT("OpenManifestFolderButton_Tooltip", "Opens the folder containing the build files in Explorer or Finder (it's recommended you check these in to source control to share with your team)"))
+				.OnClicked(this, &FLuminTargetSettingsDetails::OpenBuildFolder)
+			]
+		];
+	BuildPathPicker(DetailBuilder, AppTitleCategory, IconModelPathAttribute,
+		LOCTEXT("IconModelLabel", "Icon Model"),
+		LOCTEXT("PickIconModelButton_Tooltip", "Select the icon model to use for the application. The files will be copied to the project build folder."),
+		FOnChoosePath::CreateRaw(this, &FLuminTargetSettingsDetails::OnPickDirectory),
+		FOnPickPath::CreateRaw(this, &FLuminTargetSettingsDetails::OnPickIconModelPath));
+	BuildPathPicker(DetailBuilder, AppTitleCategory, IconPortalPathAttribute,
+		LOCTEXT("IconPortalLabel", "Icon Portal"),
+		LOCTEXT("PickIconPortalButton_Tooltip", "Select the icon portal to use for the application. The files will be copied to the project build folder."),
+		FOnChoosePath::CreateRaw(this, &FLuminTargetSettingsDetails::OnPickDirectory),
+		FOnPickPath::CreateRaw(this, &FLuminTargetSettingsDetails::OnPickIconPortalPath));
+
+	////////// UI for signing cert..
+
+	IDetailCategoryBuilder& DistributionSigningCategory = DetailBuilder.EditCategory(TEXT("Distribution Signing"));
+	DetailBuilder.HideProperty("Certificate");
+	BuildPathPicker(DetailBuilder, DistributionSigningCategory, CertificateAttribute,
+		LOCTEXT("CertificateFilePathLabel", "Certificate File Path"),
+		LOCTEXT("PickCertificateButton_Tooltip", "Select the certificate to use for signing a distribution package. The file will be copied to the project build folder."),
+		FOnChoosePath::CreateLambda([this](TAttribute<FString> FilePath, const FOnPickPath& OnPick, TSharedPtr<SButton> PickButton)->FReply
+		{
+			const FString FilterText = LOCTEXT("CertificateFile", "Certificate File").ToString();
+			return this->OnPickFile(FilePath, OnPick, PickButton,
+				LOCTEXT("PickCertificateFileDialogTitle", "Choose a certificate").ToString(),
+				FString::Printf(TEXT("%s (*.cert)|*.cert"), *FilterText));
+		}),
+		FOnPickPath::CreateRaw(this, &FLuminTargetSettingsDetails::OnPickCertificate));
+}
+
+FReply FLuminTargetSettingsDetails::OnPickIconModelPath(const FString& DirPath)
+{
+	FString ProjectModelPath = GameLuminPath / TEXT("Model");
+	if (ProjectModelPath != DirPath)
+	{
+		// Copy the contents of the selected path to the project build path.
+		CopyDir(DirPath, ProjectModelPath);
+	}
+	return FReply::Handled();
+}
+
+FReply FLuminTargetSettingsDetails::OnPickIconPortalPath(const FString& DirPath)
+{
+	FString ProjectPortalPath = GameLuminPath / TEXT("Portal");
+	if (ProjectPortalPath != DirPath)
+	{
+		// Copy the contents of the selected path to the project build path.
+		CopyDir(DirPath, ProjectPortalPath);
+	}
+	return FReply::Handled();
+}
+
+bool FLuminTargetSettingsDetails::CopyDir(FString SourceDir, FString TargetDir)
+{
+	FPaths::NormalizeDirectoryName(SourceDir);
+	FPaths::NormalizeDirectoryName(TargetDir);
+	if (!IPlatformFile::GetPlatformPhysical().DirectoryExists(*SourceDir))
+	{
+		return false;
+	}
+	// The source control utilities only deal with single files at a time. Hence need to collect
+	// the files we are copying and copy each one in turn.
+	TArray<FString> FilesToCopy;
+	int FilesCopiedCount = 0;
+	IPlatformFile::GetPlatformPhysical().FindFilesRecursively(FilesToCopy, *SourceDir, nullptr);
+	FText Description = FText::FromString(FPaths::GetBaseFilename(TargetDir));
+	for (FString& FileToCopy : FilesToCopy)
+	{
+		FString NewFile = FileToCopy.Replace(*SourceDir, *TargetDir);
+		if (IPlatformFile::GetPlatformPhysical().FileExists(*FileToCopy))
+		{
+			FString ToCopySubDir = FPaths::GetPath(FileToCopy);
+			if (!IPlatformFile::GetPlatformPhysical().DirectoryExists(*ToCopySubDir))
+			{
+				IPlatformFile::GetPlatformPhysical().CreateDirectoryTree(*ToCopySubDir);
+			}
+			FText ErrorMessage;
+			if (SourceControlHelpers::CopyFileUnderSourceControl(NewFile, FileToCopy, Description, ErrorMessage))
+			{
+				FilesCopiedCount += 1;
+			}
+			else
+			{
+				FNotificationInfo Info(ErrorMessage);
+				Info.ExpireDuration = 3.0f;
+				FSlateNotificationManager::Get().AddNotification(Info);
+			}
+		}
+	}
+	return true;
+}
+
+void FLuminTargetSettingsDetails::BuildPathPicker(
+	IDetailLayoutBuilder& DetailBuilder, IDetailCategoryBuilder& Category,
+	TAttribute<FString> Path, const FText& Label, const FText& Tooltip,
+	const FOnChoosePath& OnChoose, const FOnPickPath& OnPick)
+{
+	TSharedPtr<SButton> PickButton = nullptr;
+	TSharedPtr<SWidget> PickWidget = nullptr;
+	PickWidget = SAssignNew(PickButton, SButton)
+		.ButtonStyle(FEditorStyle::Get(), "HoverHintOnly")
+		.ToolTipText(Tooltip)
+		.ContentPadding(2.0f)
+		.ForegroundColor(FSlateColor::UseForeground())
+		.IsFocusable(false)
+		[
+			SNew(SImage)
+			.Image(FEditorStyle::GetBrush("PropertyWindow.Button_Ellipsis"))
+			.ColorAndOpacity(FSlateColor::UseForeground())
+		]
+	;
+	PickButton->SetOnClicked(FOnClicked::CreateLambda([OnChoose, OnPick, Path, PickButton]()->FReply
+	{
+		return OnChoose.Execute(Path.Get(), OnPick, PickButton);
+	}));
+
+	Category.AddCustomRow(Label, false)
+		.IsEnabled(SetupForPlatformAttribute)
+		.NameContent()
+		[
+			SNew(SHorizontalBox)
+			+ SHorizontalBox::Slot()
+			.Padding(FMargin(0, 1, 0, 1))
+			.FillWidth(1.0f)
+			[
+				SNew(STextBlock)
+				.Text(Label)
+				.Font(DetailBuilder.GetDetailFont())
+			]
+		]
+		.ValueContent()
+		[
+			SNew(SHorizontalBox)
+			+ SHorizontalBox::Slot()
+			.AutoWidth()
+			[
+				SNew(STextBlock)
+				.Text(TAttribute<FText>::Create(
+					TAttribute<FText>::FGetter::CreateLambda([Path]()->FText { return FText::FromString(Path.Get()); })))
+				.Font(DetailBuilder.GetDetailFont())
+				.Margin(2.0f)
+			]
+			+ SHorizontalBox::Slot()
+			.AutoWidth()
+			[
+				PickWidget.ToSharedRef()
+			]
+		];
+}
+
+FReply FLuminTargetSettingsDetails::OnPickDirectory(TAttribute<FString> DirPath, const FLuminTargetSettingsDetails::FOnPickPath& OnPick, TSharedPtr<SButton> PickButton)
+{
+	IDesktopPlatform* DesktopPlatform = FDesktopPlatformModule::Get();
+	if (DesktopPlatform)
+	{
+		TSharedPtr<SWindow> ParentWindow = FSlateApplication::Get().FindWidgetWindow(PickButton.ToSharedRef());
+		void* ParentWindowHandle = (ParentWindow.IsValid() && ParentWindow->GetNativeWindow().IsValid()) ? ParentWindow->GetNativeWindow()->GetOSWindowHandle() : nullptr;
+		FString StartDirectory = FEditorDirectories::Get().GetLastDirectory(ELastDirectory::GENERIC_IMPORT);
+		FString Directory;
+		if (DesktopPlatform->OpenDirectoryDialog(ParentWindowHandle, LOCTEXT("FolderDialogTitle", "Choose a directory").ToString(), StartDirectory, Directory))
+		{
+			FEditorDirectories::Get().SetLastDirectory(ELastDirectory::GENERIC_IMPORT, Directory);
+			return OnPick.Execute(Directory);
+		}
+	}
+	return FReply::Handled();
+}
+
+FReply FLuminTargetSettingsDetails::OnPickFile(
+	TAttribute<FString> FilePath, const FLuminTargetSettingsDetails::FOnPickPath& OnPick,
+	TSharedPtr<SButton> PickButton, const FString & Title, const FString & Filter)
+{
+	IDesktopPlatform* DesktopPlatform = FDesktopPlatformModule::Get();
+	if (DesktopPlatform)
+	{
+		TSharedPtr<SWindow> ParentWindow = FSlateApplication::Get().FindWidgetWindow(PickButton.ToSharedRef());
+		void* ParentWindowHandle = (ParentWindow.IsValid() && ParentWindow->GetNativeWindow().IsValid()) ? ParentWindow->GetNativeWindow()->GetOSWindowHandle() : nullptr;
+		FString StartDirectory = FEditorDirectories::Get().GetLastDirectory(ELastDirectory::GENERIC_IMPORT);
+		TArray<FString> OutFiles;
+		if (DesktopPlatform->OpenFileDialog(
+			ParentWindowHandle,
+			Title,
+			FilePath.Get(),
+			TEXT(""),
+			Filter,
+			EFileDialogFlags::None,
+			OutFiles))
+		{
+			return OnPick.Execute(OutFiles[0]);
+		}
+	}
+	return FReply::Handled();
+}
+
+FReply FLuminTargetSettingsDetails::OpenBuildFolder()
+{
+	const FString BuildFolder = FPaths::ConvertRelativePathToFull(GameLuminPath);
+	FPlatformProcess::ExploreFolder(*BuildFolder);
+
+	return FReply::Handled();
+}
+
+FReply FLuminTargetSettingsDetails::OnPickCertificate(const FString& SourceCertificateFile)
+{
+	FString TargetCertificateFile = GameLuminPath / FPaths::GetCleanFilename(SourceCertificateFile);
+	if (!IPlatformFile::GetPlatformPhysical().FileExists(*SourceCertificateFile))
+	{
+		// Sanity check for chosen file. Do nothing if it doesn't exist.
+		return FReply::Handled();
+	}
+	// We only ask for the certificate file.. But we also need the accompanying private key file.
+	FString SourceKeyFile = FPaths::Combine(FPaths::GetPath(SourceCertificateFile), FPaths::GetBaseFilename(SourceCertificateFile) + TEXT(".privkey"));
+	FString TargetKeyFile = GameLuminPath / FPaths::GetCleanFilename(SourceKeyFile);
+	if (!IPlatformFile::GetPlatformPhysical().FileExists(*SourceKeyFile))
+	{
+		// We really need the key file.
+		FNotificationInfo Info(FText(LOCTEXT("LuminMissingPrivKeyFile", "Could not find private key file.")));
+		Info.ExpireDuration = 3.0f;
+		FSlateNotificationManager::Get().AddNotification(Info);
+		return FReply::Handled();
+	}
+	{
+		FText Description = FText::FromString(FPaths::GetBaseFilename(TargetCertificateFile));
+		FText ErrorMessage;
+		if (!SourceControlHelpers::CopyFileUnderSourceControl(TargetCertificateFile, SourceCertificateFile, Description, ErrorMessage))
+		{
+			FNotificationInfo Info(ErrorMessage);
+			Info.ExpireDuration = 3.0f;
+			FSlateNotificationManager::Get().AddNotification(Info);
+		}
+	}
+	{
+		FText Description = FText::FromString(FPaths::GetBaseFilename(TargetKeyFile));
+		FText ErrorMessage;
+		if (!SourceControlHelpers::CopyFileUnderSourceControl(TargetKeyFile, SourceKeyFile, Description, ErrorMessage))
+		{
+			FNotificationInfo Info(ErrorMessage);
+			Info.ExpireDuration = 3.0f;
+			FSlateNotificationManager::Get().AddNotification(Info);
+			return FReply::Handled();
+		}
+	}
+	CertificateProp->SetValue(TargetCertificateFile.Replace(*FPaths::ProjectDir(), TEXT("")));
+	return FReply::Handled();
+}
+
 #undef LOCTEXT_NAMESPACE
diff --git a/Engine/Source/Developer/Lumin/LuminPlatformEditor/Private/LuminTargetSettingsDetails.h b/Engine/Source/Developer/Lumin/LuminPlatformEditor/Private/LuminTargetSettingsDetails.h
--- a/Engine/Source/Developer/Lumin/LuminPlatformEditor/Private/LuminTargetSettingsDetails.h
+++ b/Engine/Source/Developer/Lumin/LuminPlatformEditor/Private/LuminTargetSettingsDetails.h
@@ -4,6 +4,10 @@
 #include "PropertyHandle.h"
 #include "IDetailCustomization.h"
 #include "IPropertyTypeCustomization.h"
+#include "Engine/EngineTypes.h"
+#include "Framework/SlateDelegates.h"
+#include "Delegates/Delegate.h"
+#include "Widgets/Input/SButton.h"
 
 #include "TargetPlatformAudioCustomization.h"
 
@@ -20,6 +24,55 @@ public:
 	virtual void CustomizeDetails(IDetailLayoutBuilder& DetailBuilder) override;
 
 private:
-	
+	// These are hard-wired defaults for the engine tree default icons.
+	const FString DefaultIconModelPath;
+	const FString DefaultIconPortalPath;
+
+	const FString GameLuminPath;
+	const FString GameProjectSetupPath;
+
+	IDetailLayoutBuilder* SavedLayoutBuilder;
+
+	FLuminTargetSettingsDetails();
+
+	FUNC_DECLARE_DELEGATE(FOnPickPath, FReply, const FString&);
+	FUNC_DECLARE_DELEGATE(FOnChoosePath, FReply, TAttribute<FString>, const FOnPickPath&, TSharedPtr<SButton>);
+	void BuildPathPicker(IDetailLayoutBuilder& DetailBuilder, IDetailCategoryBuilder& Category,
+		TAttribute<FString> DirPath, const FText& Label, const FText& Tooltip,
+		const FOnChoosePath& OnChoose, const FOnPickPath& OnPick);
+	FReply OnPickDirectory(TAttribute<FString> DirPath,
+		const FOnPickPath& OnPick, TSharedPtr<SButton> PickButton);
+	FReply OnPickFile(TAttribute<FString> FilePath,
+		const FOnPickPath& OnPick, TSharedPtr<SButton> PickButton,
+		const FString & Title, const FString & Filter);
+
+	// Setup files.
+
+	TAttribute<bool> SetupForPlatformAttribute;
+
+	TSharedPtr<IPropertyHandle> IconModelPathProp;
+	TSharedPtr<IPropertyHandle> IconPortalPathProp;
+	TSharedPtr<IPropertyHandle> CertificateProp;
+
+	TAttribute<FString> IconModelPathAttribute;
+	TAttribute<FString> IconPortalPathAttribute;
+	TAttribute<FString> CertificateAttribute;
+
+	FString IconModelPathGetter();
+	FString IconPortalPathGetter();
+	FString CertificateGetter();
+
+	void BuildAppTileSection(IDetailLayoutBuilder& DetailBuilder);
+	void CopySetupFilesIntoProject();
+	FReply OpenBuildFolder();
+	FReply OnPickIconModelPath(const FString& DirPath);
+	FReply OnPickIconPortalPath(const FString& DirPath);
+	bool CopyDir(FString SourceDir, FString TargetDir);
+	FReply OnPickCertificate(const FString& CertificatePath);
+
+	// Audio section.
+
 	FAudioPluginWidgetManager AudioPluginManager;
+
+	void BuildAudioSection(IDetailLayoutBuilder& DetailBuilder);
 };
\ No newline at end of file
diff --git a/Engine/Source/Developer/Lumin/LuminPlatformEditor/Private/MagicLeapSDKSettings.h b/Engine/Source/Developer/Lumin/LuminPlatformEditor/Private/MagicLeapSDKSettings.h
--- a/Engine/Source/Developer/Lumin/LuminPlatformEditor/Private/MagicLeapSDKSettings.h
+++ b/Engine/Source/Developer/Lumin/LuminPlatformEditor/Private/MagicLeapSDKSettings.h
@@ -20,7 +20,7 @@ public:
 	UMagicLeapSDKSettings();
 	
 	/** Location on disk of the Magic Leap SDK (falls back to MLSDK environment variable if this is left blank) */
-	UPROPERTY(GlobalConfig, EditAnywhere, Category = "SDK Config", Meta = (DisplayName = "Location of Magic Leap SDK", ConfigRestartRequired = true))
+	UPROPERTY(GlobalConfig, EditAnywhere, Category = SDKConfig, Meta = (DisplayName = "Location of Magic Leap SDK"))
 	FDirectoryPath MLSDKPath;
 
 #if WITH_EDITOR
diff --git a/Engine/Source/Developer/Lumin/LuminTargetPlatform/Private/LuminTargetDevice.h b/Engine/Source/Developer/Lumin/LuminTargetPlatform/Private/LuminTargetDevice.h
--- a/Engine/Source/Developer/Lumin/LuminTargetPlatform/Private/LuminTargetDevice.h
+++ b/Engine/Source/Developer/Lumin/LuminTargetPlatform/Private/LuminTargetDevice.h
@@ -27,4 +27,123 @@ public:
 	{
 		return false;
 	}
+
+	bool GetMldbFullFilename(FString& OutFilename)
+	{
+		FString MLSDKPath = FPlatformMisc::GetEnvironmentVariable(TEXT("MLSDK"));
+#if PLATFORM_WINDOWS
+		OutFilename = FString::Printf(TEXT("%s\\tools\\mldb\\mldb.exe"), *MLSDKPath);
+#else
+		OutFilename = FString::Printf(TEXT("%s/tools/mldb/mldb"), *MLSDKPath);
+#endif
+		if (OutFilename[0] != 0)
+		{
+			return true;
+		}
+		return false;
+	}
+
+	bool ExecuteMldbCommand(const FString Command, FString* OutStdOut, FString* OutStdErr)
+	{
+		FString MldbPath;
+		if (!GetMldbFullFilename(MldbPath))
+		{
+			return false;
+		}
+
+		int32 ReturnCode;
+		FString DefaultError;
+		if (!OutStdErr)
+		{
+			OutStdErr = &DefaultError;
+		}
+		FString MldbCommand = FString::Printf(TEXT("-s %s %s"), *GetSerialNumber(), *Command);
+
+		FPlatformProcess::ExecProcess(*MldbPath, *MldbCommand, &ReturnCode, OutStdOut, OutStdErr);
+		if (ReturnCode != 0)
+		{
+			FPlatformMisc::LowLevelOutputDebugStringf(TEXT("Command %s failed with error code %d"), *MldbCommand, ReturnCode, **OutStdErr);
+			return false;
+		}
+		return true;
+	}
+
+	virtual bool SupportsFeature(ETargetDeviceFeatures Feature) const
+	{
+		switch (Feature)
+		{
+		case ETargetDeviceFeatures::MultiLaunch:
+			return false;
+		case ETargetDeviceFeatures::PowerOff:
+			return true;
+		case ETargetDeviceFeatures::PowerOn:
+			return false;
+		case ETargetDeviceFeatures::Reboot:
+			return true;
+		default:
+			return false;
+		};
+	}
+
+	bool AdviseLockStatus()
+	{
+		FString RetVal;
+		if (ExecuteMldbCommand(FString::Printf(TEXT("access-status")), &RetVal, nullptr))
+		{
+			if (RetVal.Contains(TEXT("Device locked  : True")))
+			{
+				if (FPlatformMisc::MessageBoxExt(EAppMsgType::OkCancel,
+					TEXT("Lumin device is locked. This command will take 60 seconds, during which the editor will be unresponsive."),
+					TEXT("Device Locked")) == EAppReturnType::Cancel)
+				{
+					return false;
+				}
+			}
+		}
+		else
+		{
+			// If the first mldb command doesn't work ...
+			return false;
+		}
+		return true;
+	}
+
+	virtual bool PowerOff(bool Force) override
+	{
+		if (AdviseLockStatus())
+		{
+			if (!ExecuteMldbCommand(FString::Printf(TEXT("shutdown")), nullptr, nullptr))
+			{
+				return false;
+			}
+			return true;
+		}
+		return false;
+	}
+
+	virtual bool Reboot(bool bReconnect) override
+	{
+		if (AdviseLockStatus())
+		{
+			if (!ExecuteMldbCommand(FString::Printf(TEXT("reboot")), nullptr, nullptr))
+			{
+				return false;
+			}
+			return true;
+		}
+		return false;
+	}
+
+	virtual FString GetOperatingSystemName() override
+	{
+		if (!AndroidVersionString.IsEmpty())
+		{
+			return FString::Printf(TEXT("Lumin OS %s"), *AndroidVersionString);
+		}
+		else
+		{
+			return TEXT("Lumin OS");
+		}
+	}
+
 };
diff --git a/Engine/Source/Programs/AutomationTool/Lumin/LuminPlatform.Automation.cs b/Engine/Source/Programs/AutomationTool/Lumin/LuminPlatform.Automation.cs
--- a/Engine/Source/Programs/AutomationTool/Lumin/LuminPlatform.Automation.cs
+++ b/Engine/Source/Programs/AutomationTool/Lumin/LuminPlatform.Automation.cs
@@ -71,7 +71,8 @@ public class LuminPlatform : Platform
 		string Value;
 		Ini.GetString("/Script/LuminRuntimeSettings.LuminRuntimeSettings", ConfigPropertyName, out Value);
 		Value = CleanFilePath(Value);
-		DirectoryReference IconDir = GetFullPathFromEngineRelativePath(Value, SC);
+		// We can have two different kinds of paths in the ini for icons: engine exec relative or project root relative.
+		DirectoryReference IconDir = GetFullPathFromRelativePath(Value, SC);
 		List<FileReference> IconFiles = SC.FindFilesToStage(IconDir, "*", StageFilesSearch.AllDirectories);
 		StringBuilder Builder = new StringBuilder();
 		foreach(FileReference IconFile in IconFiles)
@@ -200,26 +201,34 @@ public class LuminPlatform : Platform
 		}
 	}
 
-	private DirectoryReference GetFullPathFromEngineRelativePath(string RelativeToEnginePath, DeploymentContext SC)
+	private DirectoryReference GetFullPathFromRelativePath(string RelativePath, DeploymentContext SC)
 	{
-		string fullPath = RelativeToEnginePath;
-		if (!string.IsNullOrEmpty(fullPath))
+		string fullPath = RelativePath;
+		if (!string.IsNullOrEmpty(fullPath) && !(Path.IsPathRooted(fullPath)))
 		{
-			if (CleanFilePath(Path.GetFullPath(fullPath)) == CleanFilePath(fullPath))
+			if (Path.IsPathRooted(fullPath))
 			{
-				// Staging a complete folder required backslash on Windows.
+				// We where handed an absolute path. So just use that.
 				fullPath = Path.GetFullPath(fullPath);
 			}
 			else
 			{
-				// Paths used for icon are relative to the engines platform specific binary location, or absolute.
-				string dir = Path.GetFullPath(CombinePaths(SC.LocalRoot.FullName, "Engine/Binaries/Lumin", Path.GetDirectoryName(fullPath)));
-				string file = Path.GetFileName(fullPath);
-				// Update fullPath with corrected directory
-				fullPath = CombinePaths(dir, file);
+				// For relative paths we need to figure out if they are in the engine or project tree.
+				string FromProjectFullPath = Path.GetFullPath(CombinePaths(SC.ProjectRoot.ToString(), CleanFilePath(RelativePath)));
+				string FromEngineFullPath = Path.GetFullPath(CombinePaths(SC.EngineRoot.ToString(), CleanFilePath(RelativePath)));
+				if (Directory.Exists(FromProjectFullPath) || File.Exists(FromProjectFullPath))
+				{
+					// Works as a project relative path.. We'll go with that.
+					fullPath = FromProjectFullPath;
+				}
+				else
+				{
+					// Not in the project tree.. Assume it's in the engine tree.
+					fullPath = FromEngineFullPath;
+				}
 			}
 		}
-		return new DirectoryReference(fullPath);
+		return new DirectoryReference(Path.GetFullPath(fullPath));
 	}
 
 	private void StageNvTegraGfxDebugger(DeploymentContext SC, string Dir)
@@ -347,7 +356,7 @@ public class LuminPlatform : Platform
 				"@if \"%ERRORLEVEL%\" NEQ \"0\" goto OobeError",
 				"goto:eof",
 				":OobeError",
-				"@echo Device is not ready for use. Run \"%MLSDK% ps\" from a command prompt for details.",
+				"@echo Device is not ready for use. Run \"%MLDB% ps\" from a command prompt for details.",
 				"goto Pause",
 				":Error",
 				"@echo.",
@@ -355,7 +364,7 @@ public class LuminPlatform : Platform
 				"@echo.",
 				"@echo Things to try:",
 				"@echo Check that the device (and only the device) is listed with \"%MLDB% devices\" from a command prompt.",
-				"@echo Check if the device is ready for use with \"%MLSDK% ps\" from a command prompt.",
+				"@echo Check if the device is ready for use with \"%MLDB% ps\" from a command prompt.",
 				":Pause",
 				"@pause"
 			};
@@ -379,7 +388,7 @@ public class LuminPlatform : Platform
 					"\techo",
 					"\techo \"Things to try:\"",
 					"\techo \"Check that the device (and only the device) is listed with \"$MLDB devices\" from a command prompt.\"",
-					"\techo \"Check if the device is ready for use with \"%MLSDK% ps\" from a command prompt.\"",
+					"\techo \"Check if the device is ready for use with \"%MLDB% ps\" from a command prompt.\"",
 					"\techo",
 					"\texit 1",
 				"fi",
@@ -387,7 +396,7 @@ public class LuminPlatform : Platform
 				"echo \"Installation successful\"",
 				"$MLDB $DEVICE ps > /dev/null",
 				"if [ $? -ne 0 ]; then",
-					"\techo \"Device is not ready for use. Run \"%MLSDK% ps\" from a command prompt for details.\"",				
+					"\techo \"Device is not ready for use. Run \"%MLDB% ps\" from a command prompt for details.\"",				
 				"fi",
 				"exit 0",
 			};
@@ -1187,12 +1196,12 @@ public class LuminPlatform : Platform
 	{
 		var AppGPUArchitectures = LuminExports.CreateToolChain(Params.RawProjectPath).GetAllGPUArchitectures();
 
-		if (AppGPUArchitectures.Contains("-gl4"))
+		if (AppGPUArchitectures.Contains("-lumingl4"))
 		{
-			return "-gl4";
+			return "-lumingl4";
 		}
 
-		return "-es2";
+		return "-lumin";
 	}
 
 	private List<string> GetConfigurations(DeploymentContext SC)
diff --git a/Engine/Source/Programs/UnrealBuildTool/Platform/Android/AndroidToolChain.cs b/Engine/Source/Programs/UnrealBuildTool/Platform/Android/AndroidToolChain.cs
--- a/Engine/Source/Programs/UnrealBuildTool/Platform/Android/AndroidToolChain.cs
+++ b/Engine/Source/Programs/UnrealBuildTool/Platform/Android/AndroidToolChain.cs
@@ -35,7 +35,8 @@ namespace UnrealBuildTool
 			{ "-x64", "x6" },
 			{ "-es2", "" }, // since there's only one gpu arch now, we can strip it
 			//LUMIN_MERGE
-			{ "-gl4", "" }
+			{ "-lumingl4", "" },
+			{ "-lumin", "" }
 		};
 
 
@@ -682,9 +683,9 @@ namespace UnrealBuildTool
 				Result += " -fno-exceptions";
 			}
 
-			// Conditionally enable (default disabled) generation of information about every class with virtual functions for use by the C++ runtime type identification features 
-			// (`dynamic_cast' and `typeid'). If you don't use those parts of the language, you can save some space by using -fno-rtti. 
-			// Note that exception handling uses the same information, but it will generate it as needed. 
+			// Conditionally enable (default disabled) generation of information about every class with virtual functions for use by the C++ runtime type identification features
+			// (`dynamic_cast' and `typeid'). If you don't use those parts of the language, you can save some space by using -fno-rtti.
+			// Note that exception handling uses the same information, but it will generate it as needed.
 			if (CompileEnvironment.bUseRTTI)
 			{
 				Result += " -frtti";
@@ -697,8 +698,8 @@ namespace UnrealBuildTool
 			//@todo android: these are copied verbatim from UE3 and probably need adjustment
 			if (Architecture == "-armv7")
 			{
-				//		Result += " -mthumb-interwork";			// Generates code which supports calling between ARM and Thumb instructions, w/o it you can't reliability use both together 
-				Result += " -funwind-tables";           // Just generates any needed static data, affects no code 
+				//		Result += " -mthumb-interwork";			// Generates code which supports calling between ARM and Thumb instructions, w/o it you can't reliability use both together
+				Result += " -funwind-tables";           // Just generates any needed static data, affects no code
 				Result += " -fstack-protector";         // Emits extra code to check for buffer overflows
 														//		Result += " -mlong-calls";				// Perform function calls by first loading the address of the function into a reg and then performing the subroutine call
 				Result += " -fno-strict-aliasing";      // Prevents unwanted or invalid optimizations that could produce incorrect code
@@ -719,7 +720,7 @@ namespace UnrealBuildTool
 					Result += " -mfpu=vfpv3-d16";       //@todo android: UE3 was just vfp. arm7a should all support v3 with 16 registers
 				}
 
-				// Add flags for on-device debugging	
+				// Add flags for on-device debugging
 				if (CompileEnvironment.Configuration == CppConfiguration.Debug)
 				{
 					Result += " -fno-omit-frame-pointer";   // Disable removing the save/restore frame pointer for better debugging
@@ -740,7 +741,7 @@ namespace UnrealBuildTool
 			}
 			else if (Architecture == "-arm64")
 			{
-				Result += " -funwind-tables";           // Just generates any needed static data, affects no code 
+				Result += " -funwind-tables";           // Just generates any needed static data, affects no code
 				Result += " -fstack-protector";         // Emits extra code to check for buffer overflows
 				Result += " -fno-strict-aliasing";      // Prevents unwanted or invalid optimizations that could produce incorrect code
 				Result += " -fpic";                     // Generates position-independent code (PIC) suitable for use in a shared library
diff --git a/Engine/Source/Programs/UnrealBuildTool/Platform/Lumin/LuminProjectGenerator.cs b/Engine/Source/Programs/UnrealBuildTool/Platform/Lumin/LuminProjectGenerator.cs
--- a/Engine/Source/Programs/UnrealBuildTool/Platform/Lumin/LuminProjectGenerator.cs
+++ b/Engine/Source/Programs/UnrealBuildTool/Platform/Lumin/LuminProjectGenerator.cs
@@ -82,7 +82,7 @@ namespace UnrealBuildTool
 
 			if (IsVSLuminSupportInstalled(InProjectFileFormat) && TargetType == TargetType.Game && InPlatform == UnrealTargetPlatform.Lumin)
 			{
-				string MLSDK = Environment.GetEnvironmentVariable("MLSDK");
+				string MLSDK = Utils.CleanDirectorySeparators(Environment.GetEnvironmentVariable("MLSDK"), '\\');
 
 				// TODO: Check if MPK name can be other than the project name.
 				string GameName = TargetRulesPath.GetFileNameWithoutExtension();
@@ -101,6 +101,9 @@ namespace UnrealBuildTool
                 ELFFile = Path.Combine(ELFFile, "..\\..\\Intermediate\\Lumin\\Mabu\\Packaged\\bin\\" + GameName);
 				string DebuggerFlavor = "MLDebugger";
 
+				string SymFile = Utils.MakePathRelativeTo(NMakeOutputPath.Directory.FullName, ProjectFilePath.Directory.FullName);
+				SymFile = Path.Combine(SymFile, "..\\..\\Intermediate\\Lumin\\Mabu\\Binaries", Path.ChangeExtension(NMakeOutputPath.GetFileName(), ".sym"));
+
 				// following are defaults for debugger options
 				string Attach = "false";
 				string EnableAutoStop = "true";
@@ -116,9 +119,43 @@ namespace UnrealBuildTool
 													"<EnableAutoStop>{5}</EnableAutoStop>" + ProjectFileGenerator.NewLine +
 													"<AutoStopAtFunction>{6}</AutoStopAtFunction>" + ProjectFileGenerator.NewLine +
 													"<EnablePrettyPrinting>{7}</EnablePrettyPrinting>" + ProjectFileGenerator.NewLine +
-													"<MLDownloadOnStart>{8}</MLDownloadOnStart>" + ProjectFileGenerator.NewLine;
-				ProjectFileBuilder.Append(ProjectFileGenerator.NewLine + string.Format(CustomPathEntriesTemplate, MLSDK, PackageFile, ELFFile, DebuggerFlavor, Attach, EnableAutoStop, AutoStopAtFunction, EnablePrettyPrinting, MLDownloadOnStart));
+													"<MLDownloadOnStart>{8}</MLDownloadOnStart>" + ProjectFileGenerator.NewLine +
+													"<SymFile>{9}</SymFile>" + ProjectFileGenerator.NewLine;
+				ProjectFileBuilder.Append(ProjectFileGenerator.NewLine + string.Format(CustomPathEntriesTemplate, MLSDK, PackageFile, ELFFile, DebuggerFlavor, Attach, EnableAutoStop, AutoStopAtFunction, EnablePrettyPrinting, MLDownloadOnStart, SymFile));
+			}
+		}
+
+		public override string GetBuildCommandEntry(ProjectFile ProjectFile, DirectoryReference BatchFileDirectory, ScriptArguments FileArguments, FileReference OnlyGameProject)
+		{
+			if (OnlyGameProject != null && File.Exists(OnlyGameProject.FullName))
+			{
+				return ProjectFile.EscapePath(ProjectFile.NormalizeProjectPath(FileReference.Combine(BatchFileDirectory, "RunUAT.bat"))) + GetBatchFileArguments(FileArguments);
+			}
+			return ProjectFile.EscapePath(ProjectFile.NormalizeProjectPath(FileReference.Combine(BatchFileDirectory, "Build.bat"))) + base.GetBatchFileArguments(FileArguments);
+		}
+
+		public override string GetReBuildCommandEntry(ProjectFile ProjectFile, DirectoryReference BatchFileDirectory, ScriptArguments FileArguments, FileReference OnlyGameProject)
+		{
+			if (OnlyGameProject != null && File.Exists(OnlyGameProject.FullName))
+			{
+				return GetBuildCommandEntry(ProjectFile, BatchFileDirectory, FileArguments, OnlyGameProject) + " -clean";
 			}
+			return ProjectFile.EscapePath(ProjectFile.NormalizeProjectPath(FileReference.Combine(BatchFileDirectory, "Rebuild.bat"))) + base.GetBatchFileArguments(FileArguments);
+		}
+
+		public override string GetCleanCommandEntry(ProjectFile ProjectFile, DirectoryReference BatchFileDirectory, ScriptArguments FileArguments, FileReference OnlyGameProject)
+		{
+			return ProjectFile.EscapePath(ProjectFile.NormalizeProjectPath(FileReference.Combine(BatchFileDirectory, "Clean.bat"))) + base.GetBatchFileArguments(FileArguments);
+		}
+
+		public override string GetBatchFileArguments(ScriptArguments FileArguments)
+		{
+			return String.Format(" BuildCookRun -nocompileeditor -nop4 -project=\"$(SolutionDir)$(ProjectName).uproject\" -cook -fastcook -stage -archive -archivedirectory=\"$(SolutionDir)..\\..\\Binaries\\Lumin\" -package -clientconfig={0} -ue4exe=UE4Editor-CMD.exe -SkipCookingEditorContent -pak -prereqs -targetplatform=Lumin -build -utf8output -compile", FileArguments.ConfigurationName);
+		}
+
+		public override FileReference GetNMakeOutputPath(FileReference InNMakePath)
+		{
+			return new FileReference(Path.Combine(InNMakePath.Directory.FullName, InNMakePath.GetFileNameWithoutExtension()+"-arm64-lumin"+InNMakePath.GetExtension()));
 		}
 	}
 }
diff --git a/Engine/Source/Programs/UnrealBuildTool/Platform/Lumin/LuminToolChain.cs b/Engine/Source/Programs/UnrealBuildTool/Platform/Lumin/LuminToolChain.cs
--- a/Engine/Source/Programs/UnrealBuildTool/Platform/Lumin/LuminToolChain.cs
+++ b/Engine/Source/Programs/UnrealBuildTool/Platform/Lumin/LuminToolChain.cs
@@ -109,14 +109,16 @@ namespace UnrealBuildTool
 						throw new BuildException("Unexpected output from mabu in --print-tools: '{0}'", Line);
 					}
 
-					ToolsDict.Add(Split[0].Trim(), "\"" + Split[1].Trim() + "\"");
+					ToolsDict.Add(Split[0].Trim(), Split[1].Trim());
 				}
 			}
 
 			// set up the path to our toolchains
-			ClangPath = ToolsDict["CXX"];
-			ArPathArm64 = ToolsDict["AR"];
+			ClangPath = "\"" + ToolsDict["CXX"] + "\"";
+			ArPathArm64 = "\"" + ToolsDict["AR"] + "\"";
+			// The strip command does not execute through the shell. Hence we don't quote it.
 			StripPath = ToolsDict["STRIP"];
+			// The objcopy command does not execute through the shell. Hence we don't quote it.
 			ObjCopyPath = ToolsDict["OBJCOPY"];
 
 			// force the compiler to be executed through a command prompt instance
@@ -161,19 +163,25 @@ namespace UnrealBuildTool
 			GPUArchitectures = new List<string>();
 			if (AdditionalGPUArches != null)
 			{
-				if (AdditionalGPUArches.Contains("es2"))
+				if (!UseMobileRendering() && AdditionalGPUArches.Contains("lumingl4"))
 				{
-					GPUArchitectures.Add("-es2");
+					GPUArchitectures.Add("-lumingl4");
 				}
-				if (AdditionalGPUArches.Contains("gl4"))
+				else
 				{
-					GPUArchitectures.Add("-gl4");
+					GPUArchitectures.Add("-lumin");
 				}
 			}
-			//if nothing was specified via commandline, use the ini settings
 			if (GPUArchitectures.Count == 0)
 			{
-				GPUArchitectures.Add((UseMobileRendering() ? "-es2" : "-gl4"));
+				if (UseMobileRendering() || UseVulkan())
+				{
+					GPUArchitectures.Add("-lumin");
+				}
+				else
+				{
+					GPUArchitectures.Add("-lumingl4");
+				}
 			}
 
 			AllComboNames = (from Arch in Arches
@@ -199,7 +207,14 @@ namespace UnrealBuildTool
 			string Params = base.GetCLArguments_Global(CompileEnvironment, Architecture);
 
 			Params += " -Wno-undefined-var-template";
-			Params += " -DPLATFORM_LUMINGL4=" + ((GPUArchitectures[0] == "-gl4") ? "1" : "0");
+			if (GPUArchitectures.Contains("-lumingl4"))
+			{
+				Params += " -DPLATFORM_LUMINGL4=1";
+			}
+			else
+			{
+				Params += " -DPLATFORM_LUMINGL4=0";
+			}
 
 			// jf: added for seal, as XGE seems to not preserve case in includes properly
 			Params += " -Wno-nonportable-include-path";         // not all of these are real
@@ -305,5 +320,35 @@ namespace UnrealBuildTool
 		{
 			return StripPath;
 		}
+
+		/// <summary>
+		/// Creates an object file with only the symbolic debug information from an executable.
+		/// </summary>
+		/// <param name="SourceExeFile">The executable with debug symbol information.</param>
+		/// <param name="TargetSymFile">The generated object file with debug symbols.</param>
+		public void ExtractSymbols(FileReference SourceExeFile, FileReference TargetSymFile)
+		{
+			ProcessStartInfo StartInfo = new ProcessStartInfo();
+			StartInfo.FileName =ObjCopyPath;
+			StartInfo.Arguments = " --only-keep-debug \"" + SourceExeFile.FullName + "\" \"" + TargetSymFile.FullName + "\"";
+			StartInfo.UseShellExecute = false;
+			StartInfo.CreateNoWindow = true;
+			Utils.RunLocalProcessAndLogOutput(StartInfo);
+		}
+
+		/// <summary>
+		/// Creates a debugger link in an executable referencing where the debug symbols for it are located.
+		/// </summary>
+		/// <param name="SourceDebugFile">An object file with the debug symbols, can be an executable or the file generated with ExtractSymbols.</param>
+		/// <param name="TargetExeFile">The executable to reference the split debug info into.</param>
+		public void LinkSymbols(FileReference SourceDebugFile, FileReference TargetExeFile)
+		{
+			ProcessStartInfo StartInfo = new ProcessStartInfo();
+			StartInfo.FileName = ObjCopyPath;
+			StartInfo.Arguments = " --add-gnu-debuglink=\"" + SourceDebugFile.FullName + "\" \"" + TargetExeFile.FullName + "\"";
+			StartInfo.UseShellExecute = false;
+			StartInfo.CreateNoWindow = true;
+			Utils.RunLocalProcessAndLogOutput(StartInfo);
+		}
 	};
-}
\ No newline at end of file
+}
diff --git a/Engine/Source/Programs/UnrealBuildTool/Platform/Lumin/UEBuildLumin.cs b/Engine/Source/Programs/UnrealBuildTool/Platform/Lumin/UEBuildLumin.cs
--- a/Engine/Source/Programs/UnrealBuildTool/Platform/Lumin/UEBuildLumin.cs
+++ b/Engine/Source/Programs/UnrealBuildTool/Platform/Lumin/UEBuildLumin.cs
@@ -314,34 +314,11 @@ namespace UnrealBuildTool
 		protected override bool HasAnySDK()
 		{
 			string EnvVarKey = "MLSDK";
-			string IniVarKey = "MLSDKPath";
 
 			string MLSDKPath = Environment.GetEnvironmentVariable(EnvVarKey);
-			{
-				var configCacheIni = ConfigCache.ReadHierarchy(ConfigHierarchyType.Engine, (DirectoryReference)null, UnrealTargetPlatform.Unknown);
-
-				string path;
-				if (GetPath(configCacheIni, "/Script/LuminPlatformEditor.MagicLeapSDKSettings", IniVarKey, out path) && !string.IsNullOrEmpty(path))
-				{
-					// if the folder specified by the config var doesn't exist, fall back to the env var.
-					if (Directory.Exists(path))
-					{
-						if (MLSDKPath != path)
-						{
-							Console.WriteLine("*** MLSDK environment variable differs from config file; overriding environment variable ***");
-						}
-
-						MLSDKPath = path;
-					}
-				}
-
-				// Set for the process
-				Environment.SetEnvironmentVariable(EnvVarKey, MLSDKPath);
-			}
-
-			// we don't have an MLSDK directory specified
 			if (String.IsNullOrEmpty(MLSDKPath))
 			{
+				Console.WriteLine("*** Unable to determine MLSDK location ***");
 				return false;
 			}
 
@@ -366,11 +343,12 @@ namespace UnrealBuildTool
 
 			if (DetectedMajorVersion < MinimumSDKVersionMajor || (DetectedMajorVersion == MinimumSDKVersionMajor && DetectedMinorVersion < MinimumSDKVersionMinor))
 			{
-				Console.WriteLine("*** Found installed MLSDK version {0}.{1} but require at least {2}.{3} ***",
-					DetectedMajorVersion, DetectedMinorVersion, MinimumSDKVersionMajor, MinimumSDKVersionMinor);
+				Console.WriteLine("*** Found installed MLSDK version {0}.{1} at '{2}' but require at least {3}.{4} ***",
+					DetectedMajorVersion, DetectedMinorVersion, MLSDKPath, MinimumSDKVersionMajor, MinimumSDKVersionMinor);
 				return false;
 			}
 
+			Console.WriteLine("*** Found installed MLSDK version {0}.{1} at '{2}' ***", DetectedMajorVersion, DetectedMinorVersion, MLSDKPath);
 			return true;
 		}
 
