Fix for ML Screen Plugin not having proper include in non-unity

diff --git a/Engine/Plugins/Lumin/MagicLeapScreens/Source/MagicLeapScreens/Private/MagicLeapScreensPlugin.cpp b/Engine/Plugins/Lumin/MagicLeapScreens/Source/MagicLeapScreens/Private/MagicLeapScreensPlugin.cpp
--- a/Engine/Plugins/Lumin/MagicLeapScreens/Source/MagicLeapScreens/Private/MagicLeapScreensPlugin.cpp
+++ b/Engine/Plugins/Lumin/MagicLeapScreens/Source/MagicLeapScreens/Private/MagicLeapScreensPlugin.cpp
@@ -5,8 +5,9 @@
 #include "IMagicLeapPlugin.h"
 #include "HeadMountedDisplayFunctionLibrary.h"
 #include "MagicLeapScreensWorker.h"
+#include "Misc/CoreDelegates.h"
 #include "MagicLeapHMD.h"
-
+ 
 #define MAX_TEXTURE_SIZE 450 * 450 * 4 // currently limited by binder implementation
 
 DEFINE_LOG_CATEGORY_STATIC(LogScreensPlugin, Display, All);
