Copying //UE4/Partner-MagicLeap-Staging to Dev-VR (//UE4/Dev-VR)

diff --git a/Engine/Plugins/Lumin/MagicLeapMusicService/Config/BaseMagicLeapMusicService.ini b/Engine/Plugins/Lumin/MagicLeapMusicService/Config/BaseMagicLeapMusicService.ini
new file mode 100644
--- /dev/null
+++ b/Engine/Plugins/Lumin/MagicLeapMusicService/Config/BaseMagicLeapMusicService.ini
@@ -0,0 +1,16 @@
+[CoreRedirects]
++EnumRedirects=(OldName="/Script/MagicLeapMusicService.EServiceProviderError", NewName="/Script/MagicLeapMusicService.EMagicLeapMusicServiceProviderError")
++EnumRedirects=(OldName="/Script/MagicLeapMusicService.EPlaybackState", NewName="/Script/MagicLeapMusicService.EMagicLeapMusicServicePlaybackState")
++EnumRedirects=(OldName="/Script/MagicLeapMusicService.EPlaybackShuffleState", NewName="/Script/MagicLeapMusicService.EMagicLeapMusicServicePlaybackShuffleState")
++EnumRedirects=(OldName="/Script/MagicLeapMusicService.EPlaybackRepeatState", NewName="/Script/MagicLeapMusicService.EMagicLeapMusicServicePlaybackRepeatState")
++EnumRedirects=(OldName="/Script/MagicLeapMusicService.EMusicServiceStatus", NewName="/Script/MagicLeapMusicService.EMagicLeapMusicServiceStatus")
++StructRedirects=(OldName="/Script/MagicLeapMusicService.TrackMetadata", NewName="/Script/MagicLeapMusicService.MagicLeapMusicServiceTrackMetadata")
++ClassRedirects=(OldName="/Script/MagicLeapMusicService.MusicServicePlaybackStateDelegate", NewName="/Script/MagicLeapMusicService.MagicLeapMusicServicePlaybackStateDelegate")
++ClassRedirects=(OldName="/Script/MagicLeapMusicService.MusicServiceRepeatStateDelegate", NewName="/Script/MagicLeapMusicService.MagicLeapMusicServiceRepeatStateDelegate")
++ClassRedirects=(OldName="/Script/MagicLeapMusicService.MusicServiceShuffleStateDelegate", NewName="/Script/MagicLeapMusicService.MagicLeapMusicServiceShuffleStateDelegate")
++ClassRedirects=(OldName="/Script/MagicLeapMusicService.MusicServiceMetadataDelegate", NewName="/Script/MagicLeapMusicService.MagicLeapMusicServiceMetadataDelegate")
++ClassRedirects=(OldName="/Script/MagicLeapMusicService.MusicServicePositionDelegate", NewName="/Script/MagicLeapMusicService.MagicLeapMusicServicePositionDelegate")
++ClassRedirects=(OldName="/Script/MagicLeapMusicService.MusicServiceErrorDelegate", NewName="/Script/MagicLeapMusicService.MagicLeapMusicServiceErrorDelegate")
++ClassRedirects=(OldName="/Script/MagicLeapMusicService.MusicServiceStatusDelegate", NewName="/Script/MagicLeapMusicService.MagicLeapMusicServiceStatusDelegate")
++ClassRedirects=(OldName="/Script/MagicLeapMusicService.MusicServiceVolumeDelegate", NewName="/Script/MagicLeapMusicService.MagicLeapMusicServiceVolumeDelegate")
++StructRedirects=(OldName="/Script/MagicLeapMusicService.MusicServiceCallbacks", NewName="/Script/MagicLeapMusicService.MagicLeapMusicServiceCallbacks")
diff --git a/Engine/Plugins/Lumin/MagicLeapMusicService/MagicLeapMusicService.uplugin b/Engine/Plugins/Lumin/MagicLeapMusicService/MagicLeapMusicService.uplugin
new file mode 100644
--- /dev/null
+++ b/Engine/Plugins/Lumin/MagicLeapMusicService/MagicLeapMusicService.uplugin
@@ -0,0 +1,43 @@
+{
+	"FileVersion" : 3,
+	"Version" : 1,
+	"VersionName": "1.0",
+	"FriendlyName" : "Magic Leap Music Service",
+	"CreatedBy" : "Magic Leap, Inc.",
+	"CreatedByURL" : "http://magicleap.com",
+	"Description" : "Support for the Magic Leap Music Service.",
+	"Category" : "Media Players",
+	"EnabledByDefault" : false,
+	"DocsURL" : "",
+	"MarketplaceURL" : "",
+	"SupportURL" : "",
+	"CanContainContent" : false,
+	"IsBetaVersion" : false,
+	"Installed" : false,
+	"SupportedTargetPlatforms": [
+		"Lumin",
+		"Win64",
+		"Mac"
+	],
+	"Modules" :
+	[
+		{
+			"Name" : "MagicLeapMusicService",
+			"Type" : "Runtime",
+			"LoadingPhase" : "PostConfigInit",
+			"WhitelistPlatforms" :
+			[
+				"Win64",
+				"Mac",
+				"Lumin"
+			]
+		}
+	],
+	"Plugins":
+	[
+		{
+			"Name": "MLSDK",
+			"Enabled": true
+		}
+	]
+}
diff --git a/Engine/Plugins/Lumin/MagicLeapMusicService/Source/MagicLeapMusicService/MagicLeapMusicService.Build.cs b/Engine/Plugins/Lumin/MagicLeapMusicService/Source/MagicLeapMusicService/MagicLeapMusicService.Build.cs
new file mode 100644
--- /dev/null
+++ b/Engine/Plugins/Lumin/MagicLeapMusicService/Source/MagicLeapMusicService/MagicLeapMusicService.Build.cs
@@ -0,0 +1,32 @@
+// Copyright 1998-2019 Epic Games, Inc. All Rights Reserved.
+
+using System.IO;
+
+namespace UnrealBuildTool.Rules
+{
+	public class MagicLeapMusicService : ModuleRules
+	{
+		public MagicLeapMusicService( ReadOnlyTargetRules Target ) : base(Target)
+		{
+			PublicIncludePaths.Add(Path.Combine(ModuleDirectory, "Public"));
+
+			PublicDependencyModuleNames.AddRange(
+				new string[]
+				{
+					"Core",
+					"CoreUObject",
+					"Engine"
+				}
+			);
+
+			PrivateDependencyModuleNames.AddRange(
+				new string[]
+				{
+					"LuminRuntimeSettings",
+					"MLSDK",
+				}
+			);
+
+		}
+	}
+}
diff --git a/Engine/Plugins/Lumin/MagicLeapMusicService/Source/MagicLeapMusicService/Private/MagicLeapMusicServiceFunctionLibrary.cpp b/Engine/Plugins/Lumin/MagicLeapMusicService/Source/MagicLeapMusicService/Private/MagicLeapMusicServiceFunctionLibrary.cpp
new file mode 100644
--- /dev/null
+++ b/Engine/Plugins/Lumin/MagicLeapMusicService/Source/MagicLeapMusicService/Private/MagicLeapMusicServiceFunctionLibrary.cpp
@@ -0,0 +1,590 @@
+// Copyright 1998-2019 Epic Games, Inc. All Rights Reserved.
+
+#include "MagicLeapMusicServiceFunctionLibrary.h"
+#include "MagicLeapMusicServicePlugin.h"
+#include "MagicLeapPluginUtil.h"
+#include "Misc/Timespan.h"
+#include "Lumin/CAPIShims/LuminAPIMusicService.h"
+#include "Async/Async.h"
+
+DEFINE_LOG_CATEGORY_STATIC(LogMusicServiceFunctionLibrary, Display, All);
+
+#if WITH_MLSDK
+FMagicLeapMusicServiceTrackMetadata MLToUnrealServiceMetadata(const MLMusicServiceMetadata& MLMetadata)
+{
+	FMagicLeapMusicServiceTrackMetadata Metadata;
+
+	Metadata.TrackTitle = UTF8_TO_TCHAR(MLMetadata.track_title);
+	Metadata.AlbumName = UTF8_TO_TCHAR(MLMetadata.album_name);
+	Metadata.AlbumURL = UTF8_TO_TCHAR(MLMetadata.album_url);
+	Metadata.AlbumCoverURL = UTF8_TO_TCHAR(MLMetadata.album_cover_url);
+	Metadata.ArtistName = UTF8_TO_TCHAR(MLMetadata.artist_name);
+	Metadata.ArtistURL = UTF8_TO_TCHAR(MLMetadata.artist_url);
+	Metadata.TrackLength = FTimespan::FromSeconds(static_cast<double>(MLMetadata.length));
+
+	return Metadata;
+}
+
+EMagicLeapMusicServiceProviderError MLToUnrealServiceErrorType(const MLMusicServiceError& InError)
+{
+	switch (InError)
+	{
+	case MLMusicServiceError_None:
+		return EMagicLeapMusicServiceProviderError::None;
+	case MLMusicServiceError_Connectivity:
+		return EMagicLeapMusicServiceProviderError::Connectivity;
+	case MLMusicServiceError_Timeout:
+		return EMagicLeapMusicServiceProviderError::Timeout;
+	case MLMusicServiceError_GeneralPlayback:
+		return EMagicLeapMusicServiceProviderError::GeneralPlayback;
+	case MLMusicServiceError_Privilege:
+		return EMagicLeapMusicServiceProviderError::Privilege;
+	case MLMusicServiceError_ServiceSpecific:
+		return EMagicLeapMusicServiceProviderError::ServiceSpecific;
+	case MLMusicServiceError_NoMemory:
+		return EMagicLeapMusicServiceProviderError::NoMemory;
+	case MLMusicServiceError_Unspecified:
+		return EMagicLeapMusicServiceProviderError::Unspecified;
+	default:
+		return EMagicLeapMusicServiceProviderError::Unspecified;
+	}
+}
+
+EMagicLeapMusicServicePlaybackState MLToUnrealPlaybackState(const MLMusicServicePlaybackState& InState)
+{
+	switch (InState)
+	{
+	case MLMusicServicePlaybackState_Playing:
+		return EMagicLeapMusicServicePlaybackState::Playing;
+	case MLMusicServicePlaybackState_Paused:
+		return EMagicLeapMusicServicePlaybackState::Paused;
+	case MLMusicServicePlaybackState_Stopped:
+		return EMagicLeapMusicServicePlaybackState::Stopped;
+	case MLMusicServicePlaybackState_Error:
+		return EMagicLeapMusicServicePlaybackState::Error;
+	case MLMusicServicePlaybackState_Unknown:
+		return EMagicLeapMusicServicePlaybackState::Unknown;
+	default:
+		return EMagicLeapMusicServicePlaybackState::Unknown;
+	}
+}
+
+MLMusicServiceShuffleState UnrealToMLShuffleState(const EMagicLeapMusicServicePlaybackShuffleState& InState)
+{
+	switch (InState)
+	{
+	case EMagicLeapMusicServicePlaybackShuffleState::On:
+		return MLMusicServiceShuffleState_On;
+	case EMagicLeapMusicServicePlaybackShuffleState::Off:
+		return MLMusicServiceShuffleState_Off;
+	case EMagicLeapMusicServicePlaybackShuffleState::Unknown:
+		return MLMusicServiceShuffleState_Unknown;
+	default:
+		return MLMusicServiceShuffleState_Unknown;
+	};
+}
+
+EMagicLeapMusicServicePlaybackShuffleState MLToUnrealShuffleState(const MLMusicServiceShuffleState& InState)
+{
+	switch (InState)
+	{
+	case MLMusicServiceShuffleState_Off:
+		return EMagicLeapMusicServicePlaybackShuffleState::Off;
+	case MLMusicServiceShuffleState_On:
+		return EMagicLeapMusicServicePlaybackShuffleState::On;
+	case MLMusicServiceShuffleState_Unknown:
+		return EMagicLeapMusicServicePlaybackShuffleState::Unknown;
+	default:
+		return EMagicLeapMusicServicePlaybackShuffleState::Unknown;
+	};
+}
+
+MLMusicServiceRepeatState UnrealToMLRepeatState(const EMagicLeapMusicServicePlaybackRepeatState& InState)
+{
+	switch (InState)
+	{
+	case EMagicLeapMusicServicePlaybackRepeatState::Off:
+		return MLMusicServiceRepeatState_Off;
+	case EMagicLeapMusicServicePlaybackRepeatState::Song:
+		return MLMusicServiceRepeatState_Song;
+	case EMagicLeapMusicServicePlaybackRepeatState::Album:
+		return MLMusicServiceRepeatState_Album;
+	case EMagicLeapMusicServicePlaybackRepeatState::Unkown:
+		return MLMusicServiceRepeatState_Unknown;
+	default:
+		return MLMusicServiceRepeatState_Unknown;
+	};
+}
+
+EMagicLeapMusicServicePlaybackRepeatState MLToUnrealRepeatState(const MLMusicServiceRepeatState& InState)
+{
+	switch (InState)
+	{
+    case MLMusicServiceRepeatState_Off:
+		return EMagicLeapMusicServicePlaybackRepeatState::Off;
+    case MLMusicServiceRepeatState_Song:
+		return EMagicLeapMusicServicePlaybackRepeatState::Song;
+    case MLMusicServiceRepeatState_Album:
+		return EMagicLeapMusicServicePlaybackRepeatState::Album;
+    case MLMusicServiceRepeatState_Unknown:
+		return EMagicLeapMusicServicePlaybackRepeatState::Unkown;
+	default:
+		return EMagicLeapMusicServicePlaybackRepeatState::Unkown;
+	};
+}
+
+EMagicLeapMusicServiceStatus MLToUnrealServiceStatus(const MLMusicServiceStatus& InStatus)
+{
+    switch (InStatus)
+    {
+    case MLMusicServiceStatus_ContextChanged:
+        return EMagicLeapMusicServiceStatus::ContextChanged;
+    case MLMusicServiceStatus_Created:
+        return EMagicLeapMusicServiceStatus::Created;
+    case MLMusicServiceStatus_LoggedIn:
+        return EMagicLeapMusicServiceStatus::LoggedIn;
+    case MLMusicServiceStatus_LoggedOut:
+        return EMagicLeapMusicServiceStatus::LoggedOut;
+    case MLMusicServiceStatus_NextTrack:
+        return EMagicLeapMusicServiceStatus::NextTrack;
+    case MLMusicServiceStatus_PrevTrack:
+        return EMagicLeapMusicServiceStatus::PreviousTrack;
+    case MLMusicServiceStatus_TrackChanged:
+        return EMagicLeapMusicServiceStatus::TrackChanged;
+    case MLMusicServiceStatus_Unknown:
+        return EMagicLeapMusicServiceStatus::Unknown;
+	default:
+		return EMagicLeapMusicServiceStatus::Unknown;
+    };
+}
+#endif // WITH_MLSDK
+
+bool UMagicLeapMusicServiceFunctionLibrary::Connect(const FString& Name)
+{
+    bool bSuccess = false;
+#if WITH_MLSDK
+    MLResult Result = MLMusicServiceConnect(TCHAR_TO_UTF8(*Name));
+    bSuccess = (Result == MLResult_Ok);
+    UE_CLOG(bSuccess == false, LogMusicServiceFunctionLibrary, Error, TEXT("MLMusicServiceConnect failed with error %s!"), UTF8_TO_TCHAR(MLGetResultString(Result)));
+#endif // WITH_MLSDK
+    return bSuccess;
+}
+
+bool UMagicLeapMusicServiceFunctionLibrary::Disconnect()
+{
+    bool bSuccess = false;
+#if WITH_MLSDK
+    MLResult Result = MLMusicServiceDisconnect();
+    bSuccess = (Result == MLResult_Ok);
+    UE_CLOG(bSuccess == false, LogMusicServiceFunctionLibrary, Error, TEXT("MLMusicServiceDisconnect failed with error %s!"), UTF8_TO_TCHAR(MLGetResultString(Result)));
+#endif // WITH_MLSDK
+    return bSuccess;
+}
+
+bool UMagicLeapMusicServiceFunctionLibrary::SetCallbacks(const FMagicLeapMusicServiceCallbacks& Callbacks)
+{
+	bool bSuccess = false;
+#if WITH_MLSDK
+
+	//Wrap the UE delegates to provide the MLMusicService callback signature. 
+	//`data` will contain the UE object. It will be a valid pointer as long as `Disconnect` is called at shutdown.
+	MLMusicServiceCallbacks MLCallbacks;
+	MLCallbacks.on_playback_state_change = [](MLMusicServicePlaybackState state, void* data)
+	{	
+
+		auto UnrealDelegates = static_cast<FMagicLeapMusicServiceCallbacks*>(data);
+
+		AsyncTask(ENamedThreads::GameThread, [=]()
+		{
+			UnrealDelegates->PlaybackStateDelegate.ExecuteIfBound(MLToUnrealPlaybackState(state));
+		});
+
+	};
+	MLCallbacks.on_repeat_state_change = [](MLMusicServiceRepeatState state, void* data)
+	{
+
+		auto UnrealDelegates = static_cast<FMagicLeapMusicServiceCallbacks*>(data);
+		AsyncTask(ENamedThreads::GameThread, [=]()
+		{
+			UnrealDelegates->RepeatStateDelegate.ExecuteIfBound(MLToUnrealRepeatState(state));
+		});
+
+	};
+	MLCallbacks.on_shuffle_state_change = [](MLMusicServiceShuffleState state, void* data)
+	{
+
+		auto UnrealDelegates = static_cast<FMagicLeapMusicServiceCallbacks*>(data);
+		AsyncTask(ENamedThreads::GameThread, [=]()
+		{
+			UnrealDelegates->ShuffleStateDelegate.ExecuteIfBound(MLToUnrealShuffleState(state));
+		});
+
+	};
+	MLCallbacks.on_metadata_change = [](const MLMusicServiceMetadata* metadata, void* data)
+	{
+
+		auto UnrealDelegates = static_cast<FMagicLeapMusicServiceCallbacks*>(data);
+		FMagicLeapMusicServiceTrackMetadata UEMetadata = MLToUnrealServiceMetadata(*metadata);
+		AsyncTask(ENamedThreads::GameThread, [=]()
+		{
+			UnrealDelegates->MetadataDelegate.ExecuteIfBound(UEMetadata);
+		});
+
+	};
+	MLCallbacks.on_position_change = [](int32_t position, void* data)
+	{
+
+		auto UnrealDelegates = static_cast<FMagicLeapMusicServiceCallbacks*>(data);
+		AsyncTask(ENamedThreads::GameThread, [=]()
+		{
+			UnrealDelegates->PositionDelegate.ExecuteIfBound(FTimespan::FromMilliseconds(position));
+		});
+
+	};
+	MLCallbacks.on_error = [](MLMusicServiceError error_type, int32_t error_code, void* data)
+	{
+
+		auto UnrealDelegates = static_cast<FMagicLeapMusicServiceCallbacks*>(data);
+		AsyncTask(ENamedThreads::GameThread, [=]()
+		{
+			UnrealDelegates->ErrorDelegate.ExecuteIfBound(MLToUnrealServiceErrorType(error_type), error_code);
+		});
+
+	};
+	MLCallbacks.on_status_change = [](MLMusicServiceStatus status, void* data)
+	{
+
+		auto UnrealDelegates = static_cast<FMagicLeapMusicServiceCallbacks*>(data);
+		AsyncTask(ENamedThreads::GameThread, [=]()
+		{
+			UnrealDelegates->StatusDelegate.ExecuteIfBound(MLToUnrealServiceStatus(status));
+		});
+
+	};
+	MLCallbacks.on_volume_change = [](float volume, void* data)
+	{
+
+		auto UnrealDelegates = static_cast<FMagicLeapMusicServiceCallbacks*>(data);
+		AsyncTask(ENamedThreads::GameThread, [=]()
+		{
+			UnrealDelegates->VolumeDelegate.ExecuteIfBound(volume);
+		});
+
+	};
+
+	//Store the callback state
+	FMagicLeapMusicServiceCallbacks& UECallbacks = GetMagicLeapMusicServiceModule().GetDelegates();
+	UECallbacks = Callbacks;
+
+	MLResult Result = MLMusicServiceSetCallbacks(&MLCallbacks, &UECallbacks);
+	bSuccess = (Result == MLResult_Ok);
+	UE_CLOG(bSuccess == false, LogMusicServiceFunctionLibrary, Error, TEXT("MLMusicServiceSetCallbacks failed with error %s!"), UTF8_TO_TCHAR(MLGetResultString(Result)));
+#endif // WITH_MLSDK
+	return bSuccess;
+}
+
+bool UMagicLeapMusicServiceFunctionLibrary::SetAuthenticationString(const FString& AuthenticationString)
+{
+    bool bSuccess = false;
+#if WITH_MLSDK
+    MLResult Result = MLMusicServiceSetAuthString(TCHAR_TO_UTF8(*AuthenticationString));
+    bSuccess = (Result == MLResult_Ok);
+    UE_CLOG(bSuccess == false, LogMusicServiceFunctionLibrary, Error, TEXT("MLMusicServiceSetAuthString failed with error %s!"), UTF8_TO_TCHAR(MLGetResultString(Result)));
+#endif // WITH_MLSDK
+    return bSuccess;
+}
+
+bool UMagicLeapMusicServiceFunctionLibrary::SetPlaybackURL(const FString& PlaybackURL)
+{
+    bool bSuccess = false;
+#if WITH_MLSDK
+    MLResult Result = MLMusicServiceSetURL(TCHAR_TO_UTF8(*PlaybackURL));
+    bSuccess = (Result == MLResult_Ok);
+    UE_CLOG(bSuccess == false, LogMusicServiceFunctionLibrary, Error, TEXT("MLMusicServiceSetURL failed with error %s!"), UTF8_TO_TCHAR(MLGetResultString(Result)));
+#endif // WITH_MLSDK
+    return bSuccess;
+}
+
+bool UMagicLeapMusicServiceFunctionLibrary::SetPlaylist(const TArray<FString>& Playlist)
+{
+	bool bSuccess = false;
+#if WITH_MLSDK
+	if (Playlist.Num() == 0)
+	{
+		return bSuccess;
+	}
+
+	uint64 Length = Playlist.Num();
+	TArray<const char*> InPlaylist;
+	InPlaylist.Init(nullptr, Length);
+	for (int32 Index = 0; Index < Length; ++Index)
+	{
+		InPlaylist[Index] = static_cast<const char*>(FMemory::Malloc(sizeof(char) * Playlist[Index].Len() + 1));
+		FCStringAnsi::Strcpy(reinterpret_cast<ANSICHAR*>(const_cast<char*>(InPlaylist[Index])), Playlist[Index].Len(), TCHAR_TO_UTF8(*Playlist[Index]));
+	}
+	MLResult Result = MLMusicServiceSetPlayList(InPlaylist.GetData(), Length);
+	bSuccess = (Result == MLResult_Ok);
+	UE_CLOG(bSuccess == false, LogMusicServiceFunctionLibrary, Error, TEXT("MLMusicServiceSetPlayList failed with error %s!"), UTF8_TO_TCHAR(MLGetResultString(Result)));
+
+	for (int32 Index = 0; Index < Length; ++Index)
+	{
+		FMemory::Free((void*)(InPlaylist[Index]));
+	}
+
+#endif // WITH_MLSDK
+	return bSuccess;
+}
+
+bool UMagicLeapMusicServiceFunctionLibrary::StartPlayback()
+{
+    bool bSuccess = false;
+#if WITH_MLSDK
+    MLResult Result = MLMusicServiceStart();
+    bSuccess = (Result == MLResult_Ok);
+    UE_CLOG(bSuccess == false, LogMusicServiceFunctionLibrary, Error, TEXT("MLMusicServiceStart failed with error %s!"), UTF8_TO_TCHAR(MLGetResultString(Result)));
+#endif // WITH_MLSDK
+    return bSuccess;
+}
+
+bool UMagicLeapMusicServiceFunctionLibrary::StopPlayback()
+{
+    bool bSuccess = false;
+#if WITH_MLSDK
+    MLResult Result = MLMusicServiceStop();
+    bSuccess = (Result == MLResult_Ok);
+    UE_CLOG(bSuccess == false, LogMusicServiceFunctionLibrary, Error, TEXT("MLMusicServiceStop failed with error %s!"), UTF8_TO_TCHAR(MLGetResultString(Result)));
+#endif // WITH_MLSDK
+    return bSuccess;
+}
+
+bool UMagicLeapMusicServiceFunctionLibrary::PausePlayback()
+{
+    bool bSuccess = false;
+#if WITH_MLSDK
+    MLResult Result = MLMusicServicePause();
+    bSuccess = (Result == MLResult_Ok);
+    UE_CLOG(bSuccess == false, LogMusicServiceFunctionLibrary, Error, TEXT("MLMusicServicePause failed with error %s!"), UTF8_TO_TCHAR(MLGetResultString(Result)));
+#endif // WITH_MLSDK
+    return bSuccess;
+}
+
+bool UMagicLeapMusicServiceFunctionLibrary::ResumePlayback()
+{
+    bool bSuccess = false;
+#if WITH_MLSDK
+    MLResult Result = MLMusicServiceResume();
+    bSuccess = (Result == MLResult_Ok);
+    UE_CLOG(bSuccess == false, LogMusicServiceFunctionLibrary, Error, TEXT("MLMusicServiceResume failed with error %s!"), UTF8_TO_TCHAR(MLGetResultString(Result)));
+#endif // WITH_MLSDK
+    return bSuccess;
+}
+
+bool UMagicLeapMusicServiceFunctionLibrary::SeekTo(const FTimespan& Position)
+{
+    bool bSuccess = false;
+#if WITH_MLSDK
+    if (ensure(Position.GetTotalMilliseconds() < UINT32_MAX))
+    {
+		MLResult Result = MLMusicServiceSeek(static_cast<uint32>(Position.GetTotalMilliseconds()));
+		bSuccess = (Result == MLResult_Ok);
+		UE_CLOG(bSuccess == false, LogMusicServiceFunctionLibrary, Error, TEXT("MLMusicServiceSeek failed with error %s!"), UTF8_TO_TCHAR(MLGetResultString(Result)));
+    }
+#endif // WITH_MLSDK
+    return bSuccess;
+}
+
+bool UMagicLeapMusicServiceFunctionLibrary::NextTrack()
+{
+    bool bSuccess = false;
+#if WITH_MLSDK
+    MLResult Result = MLMusicServiceNext();
+    bSuccess = (Result == MLResult_Ok);
+    UE_CLOG(bSuccess == false, LogMusicServiceFunctionLibrary, Error, TEXT("MLMusicServiceNext failed with error %s!"), UTF8_TO_TCHAR(MLGetResultString(Result)));
+#endif // WITH_MLSDK
+    return bSuccess;
+}
+
+bool UMagicLeapMusicServiceFunctionLibrary::PreviousTrack()
+{
+    bool bSuccess = false;
+#if WITH_MLSDK
+    MLResult Result = MLMusicServicePrevious();
+    bSuccess = (Result == MLResult_Ok);
+    UE_CLOG(bSuccess == false, LogMusicServiceFunctionLibrary, Error, TEXT("MLMusicServicePrevious failed with error %s!"), UTF8_TO_TCHAR(MLGetResultString(Result)));
+#endif // WITH_MLSDK
+    return bSuccess;
+}
+
+bool UMagicLeapMusicServiceFunctionLibrary::GetServiceProviderError(EMagicLeapMusicServiceProviderError& ErrorType, int32& ErrorCode)
+{
+	bool bSuccess = false;
+#if WITH_MLSDK
+	MLMusicServiceError InError;
+	MLResult Result = MLMusicServiceGetError(&InError, &ErrorCode);
+	ErrorType = MLToUnrealServiceErrorType(InError);
+	bSuccess = (Result == MLResult_Ok);
+	UE_CLOG(bSuccess == false, LogMusicServiceFunctionLibrary, Error, TEXT("MLMusicServiceGetError failed with error %s!"), UTF8_TO_TCHAR(MLGetResultString(Result)));
+#endif // WITH_MLSDK
+	return bSuccess;
+}
+
+bool UMagicLeapMusicServiceFunctionLibrary::GetPlaybackState(EMagicLeapMusicServicePlaybackState& PlaybackState)
+{
+	bool bSuccess = false;
+#if WITH_MLSDK
+	MLMusicServicePlaybackState InPlaybackState;
+	MLResult Result = MLMusicServiceGetPlaybackState(&InPlaybackState);
+	PlaybackState = MLToUnrealPlaybackState(InPlaybackState);
+	bSuccess = (Result == MLResult_Ok);
+	UE_CLOG(bSuccess == false, LogMusicServiceFunctionLibrary, Error, TEXT("MLMusicServiceGetPlaybackState failed with error %s!"), UTF8_TO_TCHAR(MLGetResultString(Result)));
+#endif // WITH_MLSDK
+	return bSuccess;
+}
+
+bool UMagicLeapMusicServiceFunctionLibrary::SetPlaybackShuffleState(const EMagicLeapMusicServicePlaybackShuffleState& ShuffleState)
+{
+    bool bSuccess = false;
+#if WITH_MLSDK
+    MLResult Result = MLMusicServiceSetShuffle(UnrealToMLShuffleState(ShuffleState));
+    bSuccess = (Result == MLResult_Ok);
+    UE_CLOG(bSuccess == false, LogMusicServiceFunctionLibrary, Error, TEXT("MLMusicServiceSetShuffle failed with error %s!"), UTF8_TO_TCHAR(MLGetResultString(Result)));
+#endif // WITH_MLSDK
+    return bSuccess;
+}
+
+bool UMagicLeapMusicServiceFunctionLibrary::GetPlaybackShuffleState(EMagicLeapMusicServicePlaybackShuffleState& ShuffleState)
+{
+    bool bSuccess = false;
+#if WITH_MLSDK
+    MLMusicServiceShuffleState OutShuffleState;
+    MLResult Result = MLMusicServiceGetShuffleState(&OutShuffleState);
+    ShuffleState = MLToUnrealShuffleState(OutShuffleState);
+    bSuccess = (Result == MLResult_Ok);
+    UE_CLOG(bSuccess == false, LogMusicServiceFunctionLibrary, Error, TEXT("MLMusicServiceGetShuffle failed with error %s!"), UTF8_TO_TCHAR(MLGetResultString(Result)));
+#endif // WITH_MLSDK
+    return bSuccess;
+}
+
+bool UMagicLeapMusicServiceFunctionLibrary::SetPlaybackRepeatState(const EMagicLeapMusicServicePlaybackRepeatState& RepeatState)
+{
+    bool bSuccess = false;
+#if WITH_MLSDK
+    MLResult Result = MLMusicServiceSetRepeat(UnrealToMLRepeatState(RepeatState));
+    bSuccess = (Result == MLResult_Ok);
+    UE_CLOG(bSuccess == false, LogMusicServiceFunctionLibrary, Error, TEXT("MLMusicServiceSetRepeat failed with error %s!"), UTF8_TO_TCHAR(MLGetResultString(Result)));
+#endif // WITH_MLSDK
+    return bSuccess;
+}
+
+bool UMagicLeapMusicServiceFunctionLibrary::GetPlaybackRepeatState(EMagicLeapMusicServicePlaybackRepeatState& RepeatState)
+{
+    bool bSuccess = false;
+#if WITH_MLSDK
+    MLMusicServiceRepeatState OutRepeatState;
+    MLResult Result = MLMusicServiceGetRepeatState(&OutRepeatState);
+    RepeatState = MLToUnrealRepeatState(OutRepeatState);
+    bSuccess = (Result == MLResult_Ok);
+    UE_CLOG(bSuccess == false, LogMusicServiceFunctionLibrary, Error, TEXT("MLMusicServiceGetRepeat failed with error %s!"), UTF8_TO_TCHAR(MLGetResultString(Result)));
+#endif // WITH_MLSDK
+    return bSuccess;
+}
+
+bool UMagicLeapMusicServiceFunctionLibrary::SetPlaybackVolume(const float Volume)
+{
+    bool bSuccess = false;
+#if WITH_MLSDK
+    MLResult Result = MLMusicServiceSetVolume(FMath::Clamp<float>(Volume, 0.0f, 1.0f));
+    bSuccess = (Result == MLResult_Ok);
+    UE_CLOG(bSuccess == false, LogMusicServiceFunctionLibrary, Error, TEXT("MLMusicServiceSetVolume failed with error %s!"), UTF8_TO_TCHAR(MLGetResultString(Result)));
+#endif // WITH_MLSDK
+    return bSuccess;
+}
+
+bool UMagicLeapMusicServiceFunctionLibrary::GetPlaybackVolume(float& Volume)
+{
+    bool bSuccess = false;
+#if WITH_MLSDK
+    MLResult Result = MLMusicServiceGetVolume(&Volume);
+    bSuccess = (Result == MLResult_Ok);
+    UE_CLOG(bSuccess == false, LogMusicServiceFunctionLibrary, Error, TEXT("MLMusicServiceGetVolume failed with error %s!"), UTF8_TO_TCHAR(MLGetResultString(Result)));
+#endif // WITH_MLSDK
+    return bSuccess;
+}
+
+bool UMagicLeapMusicServiceFunctionLibrary::GetTrackLength(FTimespan& Length)
+{
+    bool bSuccess = false;
+#if WITH_MLSDK
+    uint32 OutLength;
+    MLResult Result = MLMusicServiceGetTrackLength(&OutLength);
+    bSuccess = (Result == MLResult_Ok);
+    if (bSuccess == true)
+    {
+        Length = FTimespan::FromSeconds(OutLength);
+    }
+    else
+    {
+        UE_LOG(LogMusicServiceFunctionLibrary, Error, TEXT("MLMusicServiceGetTrackLength failed with error %s!"), UTF8_TO_TCHAR(MLGetResultString(Result)));
+    }
+#endif // WITH_MLSDK
+    return bSuccess;
+}
+
+bool UMagicLeapMusicServiceFunctionLibrary::GetCurrentPosition(FTimespan& Position)
+{
+    bool bSuccess = false;
+#if WITH_MLSDK
+    uint32 OutPosition;
+    MLResult Result = MLMusicServiceGetCurrentPosition(&OutPosition);
+    bSuccess = (Result == MLResult_Ok);
+    if (bSuccess == true)
+    {
+        Position = FTimespan::FromMilliseconds(OutPosition);
+    }
+    else
+    {
+        UE_LOG(LogMusicServiceFunctionLibrary, Error, TEXT("MLMusicServiceGetCurrentPosition failed with error %s!"), UTF8_TO_TCHAR(MLGetResultString(Result)));
+    }
+#endif // WITH_MLSDK
+    return bSuccess;
+}
+
+bool UMagicLeapMusicServiceFunctionLibrary::GetServiceStatus(EMagicLeapMusicServiceStatus& Status)
+{
+    bool bSuccess = false;
+#if WITH_MLSDK
+    MLMusicServiceStatus OutStatus;
+    MLResult Result = MLMusicServiceGetStatus(&OutStatus);
+    bSuccess = (Result == MLResult_Ok);
+    if (bSuccess == true)
+    {
+		Status = MLToUnrealServiceStatus(OutStatus);
+    }
+    else
+    {
+        UE_LOG(LogMusicServiceFunctionLibrary, Error, TEXT("MLMusicServiceGetServiceStatus failed with error %s!"), UTF8_TO_TCHAR(MLGetResultString(Result)));
+    }
+#endif // WITH_MLSDK
+    return bSuccess;
+}
+
+bool UMagicLeapMusicServiceFunctionLibrary::GetCurrentTrackMetadata(FMagicLeapMusicServiceTrackMetadata& Metadata)
+{
+    bool bSuccess = false;
+#if WITH_MLSDK
+    MLMusicServiceMetadata OutMetadata;
+    MLResult Result = MLMusicServiceGetMetadata(&OutMetadata);
+    bSuccess = (Result == MLResult_Ok);
+    if (bSuccess)
+    {
+		Metadata = MLToUnrealServiceMetadata(OutMetadata);
+		MLMusicServiceReleaseMetadata(&OutMetadata);
+    }
+    else
+    {
+        UE_LOG(LogMusicServiceFunctionLibrary, Error, TEXT("MLMusicServiceGetMetadata failed with error %s!"), UTF8_TO_TCHAR(MLGetResultString(Result)));
+    }
+#endif // WITH_MLSDK
+    return bSuccess;
+}
+
diff --git a/Engine/Plugins/Lumin/MagicLeapMusicService/Source/MagicLeapMusicService/Private/MagicLeapMusicServicePlugin.cpp b/Engine/Plugins/Lumin/MagicLeapMusicService/Source/MagicLeapMusicService/Private/MagicLeapMusicServicePlugin.cpp
new file mode 100644
--- /dev/null
+++ b/Engine/Plugins/Lumin/MagicLeapMusicService/Source/MagicLeapMusicService/Private/MagicLeapMusicServicePlugin.cpp
@@ -0,0 +1,12 @@
+// Copyright 1998-2019 Epic Games, Inc. All Rights Reserved.
+
+#include "MagicLeapMusicServicePlugin.h"
+
+FMagicLeapMusicServiceCallbacks& FMagicLeapMusicServicePlugin::GetDelegates()
+{
+	return Delegates;
+}
+
+
+IMPLEMENT_MODULE(FMagicLeapMusicServicePlugin, MagicLeapMusicService);
+
diff --git a/Engine/Plugins/Lumin/MagicLeapMusicService/Source/MagicLeapMusicService/Private/MagicLeapMusicServicePlugin.h b/Engine/Plugins/Lumin/MagicLeapMusicService/Source/MagicLeapMusicService/Private/MagicLeapMusicServicePlugin.h
new file mode 100644
--- /dev/null
+++ b/Engine/Plugins/Lumin/MagicLeapMusicService/Source/MagicLeapMusicService/Private/MagicLeapMusicServicePlugin.h
@@ -0,0 +1,26 @@
+// Copyright 1998-2019 Epic Games, Inc. All Rights Reserved.
+
+#pragma once
+
+#include "CoreMinimal.h"
+#include "IMagicLeapMusicServicePlugin.h"
+#include "Lumin/CAPIShims/LuminAPI.h"
+#include "MagicLeapPluginUtil.h"
+
+class FMagicLeapMusicServicePlugin : public IMagicLeapMusicServicePlugin
+{
+public:
+
+	// Returns the callbacks associated with the music service plugin. 
+	FMagicLeapMusicServiceCallbacks& GetDelegates() override;
+
+private:
+
+	FMagicLeapMusicServiceCallbacks Delegates;
+
+};
+
+inline FMagicLeapMusicServicePlugin& GetMagicLeapMusicServiceModule()
+{
+	return FModuleManager::Get().GetModuleChecked<FMagicLeapMusicServicePlugin>("MagicLeapMusicService");
+}
diff --git a/Engine/Plugins/Lumin/MagicLeapMusicService/Source/MagicLeapMusicService/Public/IMagicLeapMusicServicePlugin.h b/Engine/Plugins/Lumin/MagicLeapMusicService/Source/MagicLeapMusicService/Public/IMagicLeapMusicServicePlugin.h
new file mode 100644
--- /dev/null
+++ b/Engine/Plugins/Lumin/MagicLeapMusicService/Source/MagicLeapMusicService/Public/IMagicLeapMusicServicePlugin.h
@@ -0,0 +1,39 @@
+// Copyright 1998-2019 Epic Games, Inc. All Rights Reserved.
+
+#pragma once
+
+#include "Modules/ModuleManager.h"
+#include "MagicLeapMusicServiceTypes.h"
+
+/**
+ * The public interface to this module.  In most cases, this interface is only public to sibling modules
+ * within this plugin.
+ */
+class IMagicLeapMusicServicePlugin : public IModuleInterface
+{
+public:
+
+	/**
+	 * Singleton-like access to this module's interface.  This is just for convenience!
+	 * Beware of calling this during the shutdown phase, though.  Your module might have been unloaded already.
+	 *
+	 * @return Returns singleton instance, loading the module on demand if needed
+	 */
+	static inline IMagicLeapMusicServicePlugin& Get()
+	{
+		return FModuleManager::LoadModuleChecked<IMagicLeapMusicServicePlugin>( "MagicLeapMusicService" );
+	}
+
+	/**
+	 * Checks to see if this module is loaded and ready.  It is only valid to call Get() if IsAvailable() returns true.
+	 *
+	 * @return True if the module is loaded and ready to use
+	 */
+	static inline bool IsAvailable()
+	{
+		return FModuleManager::Get().IsModuleLoaded( "MagicLeapMusicService" );
+	}
+
+	virtual FMagicLeapMusicServiceCallbacks& GetDelegates() = 0;
+};
+
diff --git a/Engine/Plugins/Lumin/MagicLeapMusicService/Source/MagicLeapMusicService/Public/MagicLeapMusicServiceFunctionLibrary.h b/Engine/Plugins/Lumin/MagicLeapMusicService/Source/MagicLeapMusicService/Public/MagicLeapMusicServiceFunctionLibrary.h
new file mode 100644
--- /dev/null
+++ b/Engine/Plugins/Lumin/MagicLeapMusicService/Source/MagicLeapMusicService/Public/MagicLeapMusicServiceFunctionLibrary.h
@@ -0,0 +1,235 @@
+// Copyright 19982018 Epic Games, Inc. All Rights Reserved.
+
+#pragma once
+
+#include "Engine/Engine.h"
+#include "Containers/Array.h"
+#include "Kismet/BlueprintFunctionLibrary.h"
+#include "MagicLeapMusicServiceTypes.h"
+#include "MagicLeapMusicServiceFunctionLibrary.generated.h"
+
+/**
+   Provides music playback control of a Music Service Provider
+ */
+UCLASS(ClassGroup = MagicLeap)
+class MAGICLEAPMUSICSERVICE_API UMagicLeapMusicServiceFunctionLibrary : public UBlueprintFunctionLibrary
+{
+	GENERATED_BODY()
+
+public:
+	/**
+		Connects to a music service provider
+
+		@param Name  Name of the music service provider to connect to.
+		@return True if successfully connected to the music provider with name 'Name', false otherwise
+	 */
+	UFUNCTION(BlueprintCallable, Category = "Music Service Function Library|MagicLeap")
+	static bool Connect(const FString& Name);
+
+	/**
+		Disconnects from the currently connected service provider.
+
+		@return True if successfully disconnected to the currently connected music provider, false otherwise
+	 */
+	UFUNCTION(BlueprintCallable, Category = "Music Service Function Library|MagicLeap")
+	static bool Disconnect();
+
+	/**
+		Sets callbacks the music service provider will utilize.
+
+		@return True if successful in setting the music service callbacks, false otherwise
+	 */
+	UFUNCTION(BlueprintCallable, Category = "Music Service Function Library|MagicLeap")
+	static bool SetCallbacks(const FMagicLeapMusicServiceCallbacks& Callbacks);
+
+	/**
+		Sets the authentication string used to connect to the music service provider
+
+		@param AuthenticationString  String used to connect to the music service provider
+		@return True if successfully set the authentication string, false otherwise
+	 */
+	UFUNCTION(BlueprintCallable, Category = "Music Service Function Library|MagicLeap")
+	static bool SetAuthenticationString(const FString& AuthenticationString);
+
+	/**
+		Sets the URL to playback using the music service provider
+
+		@param PlaybackURL  URL to play via the music service provider
+		@return True if successfully set the playback URL, false otherwise
+	 */
+	UFUNCTION(BlueprintCallable, Category = "Music Service Function Library|MagicLeap")
+	static bool SetPlaybackURL(const FString& PlaybackURL);
+
+	/**
+		Sets the playlist for the music service provider
+
+		@param Playlist  Array of track names to be played by the music service provider
+		@return True if successfully set the playlist
+	 */
+	UFUNCTION(BlueprintCallable, Category = "Music Service Function Library|MagicLeap")
+	static bool SetPlaylist(const TArray<FString>& Playlist);
+
+	/**
+		Starts playback of the currently selected track.
+
+		@return True if playback started successfully
+	 */
+	UFUNCTION(BlueprintCallable, Category = "Music Service Function Library|MagicLeap")
+	static bool StartPlayback();
+
+	/**
+		Stops playback of the currently selected track.
+
+		@return True if playback stopped successfully
+	 */
+	UFUNCTION(BlueprintCallable, Category = "Music Service Function Library|MagicLeap")
+	static bool StopPlayback();
+
+	/**
+		Pauses playback of the currently selected track.
+
+		@return True if playback paused successfully
+	 */
+	UFUNCTION(BlueprintCallable, Category = "Music Service Function Library|MagicLeap")
+	static bool PausePlayback();
+
+	/**
+		Resumes playback of the currently selected track.
+
+		@return True if playback resumed successfully
+	 */
+	UFUNCTION(BlueprintCallable, Category = "Music Service Function Library|MagicLeap")
+	static bool ResumePlayback();
+
+	/**
+		Seeks to a position on the currently selected track.
+
+		@param Position Position on the current track to seek to.
+		@return True if the seeking 'Position' succeeded
+	 */
+	UFUNCTION(BlueprintCallable, Category = "Music Service Function Library|MagicLeap")
+	static bool SeekTo(const FTimespan& Position);
+
+	/**
+		Advances to the next track on the music service provider's playlist.
+
+		@return True if advancing to the next track succeded
+	 */
+	UFUNCTION(BlueprintCallable, Category = "Music Service Function Library|MagicLeap")
+	static bool NextTrack();
+
+	/**
+		Rewinds to the previous track on the music provider's playlist
+
+		@return True if rewind to the previous track succeded
+	 */
+	UFUNCTION(BlueprintCallable, Category = "Music Service Function Library|MagicLeap")
+	static bool PreviousTrack();
+
+	/**
+		Gets the current error status of the music service provider.
+		If no error has occurred then this will return the error type 'None'
+
+		@return True if retreiving the providers error has succeeded
+	 */
+	UFUNCTION(BlueprintCallable, BlueprintPure, Category = "Music Service Function Library|MagicLeap")
+	static bool GetServiceProviderError(EMagicLeapMusicServiceProviderError& ErrorType, int32& ErrorCode);
+
+	/**
+		Gets the current playback state of the music service provider
+
+		@return True if retreiving the playback state succeeded
+	 */
+	UFUNCTION(BlueprintCallable, BlueprintPure, Category = "Music Service Function Library|MagicLeap")
+	static bool GetPlaybackState(EMagicLeapMusicServicePlaybackState& PlaybackState);
+
+	/**
+		Sets the current playback shuffle state
+
+		@param ShuffleState  The shuffle type to set the current shuffle state to.
+		@return True if setting the shuffle state to 'ShuffleState' succeeds
+	 */
+	UFUNCTION(BlueprintCallable, Category = "Music Service Function Library|MagicLeap")
+	static bool SetPlaybackShuffleState(const EMagicLeapMusicServicePlaybackShuffleState& ShuffleState);
+
+	/**
+		Gets the current playback shuffle state
+
+		@param ShuffleState  The current playback shuffle state the music provider is set to.
+		@return True if retrieving the shuffle state succeeded
+	 */
+	UFUNCTION(BlueprintCallable, BlueprintPure, Category = "Music Service Function Library|MagicLeap")
+	static bool GetPlaybackShuffleState(EMagicLeapMusicServicePlaybackShuffleState& ShuffleState);
+
+	/**
+		Sets the repeat state of the music service provider
+
+		@param RepeatState  Repeat state to set the music service provider to
+		@return True if setting the repeat state to 'RepeatState' succeeded
+	 */
+	UFUNCTION(BlueprintCallable, Category = "Music Service Function Library|MagicLeap")
+	static bool SetPlaybackRepeatState(const EMagicLeapMusicServicePlaybackRepeatState& RepeatState);
+
+	/**
+		Gets the current repeat state of the music service provider
+
+		@param RepeatState  The repeat state the music service provider is currently set to
+		@return True if retrieving the repeat state succeeds
+	 */
+	UFUNCTION(BlueprintCallable, BlueprintPure, Category = "Music Service Function Library|MagicLeap")
+	static bool GetPlaybackRepeatState(EMagicLeapMusicServicePlaybackRepeatState& RepeatState);
+
+	/**
+		Sets the volume of playback by the music service provider
+
+		@param Volume  Volume to set playback to, should range between 0.0 and 1.0
+		@return True if playback volume has successfully been changed to 'Volume'
+	 */
+	UFUNCTION(BlueprintCallable, Category = "Music Service Function Library|MagicLeap")
+	static bool SetPlaybackVolume(const float Volume);
+
+	/**
+		Gets the volume of playback the music service provider is currently set to
+
+		@param Volume  Volume of current playback
+		@return True if the playback volume has sucessfully been retrieved
+	 */
+	UFUNCTION(BlueprintCallable, BlueprintPure, Category = "Music Service Function Library|MagicLeap")
+	static bool GetPlaybackVolume(float& Volume);
+
+	/**
+		Gets the length of the currently selected track
+
+		@param Length  Length of the currently selected track
+		@return True if the length of the current track has been sucessfully retrieved
+	 */
+	UFUNCTION(BlueprintCallable, BlueprintPure, Category = "Music Service Function Library|MagicLeap")
+	static bool GetTrackLength(FTimespan& Length);
+
+	/**
+		Gets the current playback position of the currently selected track
+
+		@param Position Position of playback of the current track
+		@return True if retrieving the current track position succeeded
+	 */
+	UFUNCTION(BlueprintCallable, BlueprintPure, Category = "Music Service Function Library|MagicLeap")
+	static bool GetCurrentPosition(FTimespan& Position);
+
+	/**
+		Gets the current service status of the music service provider
+
+		@param Status Current service status of the music service provider
+		@return True if retrieving the service status of the current music service provider succeeded
+	 */
+	UFUNCTION(BlueprintCallable, BlueprintPure, Category = "Music Service Function Library|MagicLeap")
+	static bool GetServiceStatus(EMagicLeapMusicServiceStatus& Status);
+
+	/**
+		Gets the currently selected track's metadata
+
+		@param Metadata  Metadata of the currently selected track
+		@return True if retrieving the current track metadata succeeded
+	 */
+	UFUNCTION(BlueprintCallable, BlueprintPure, Category = "Music Service Function Library|MagicLeap")
+	static bool GetCurrentTrackMetadata(FMagicLeapMusicServiceTrackMetadata& Metadata);
+};
diff --git a/Engine/Plugins/Lumin/MagicLeapMusicService/Source/MagicLeapMusicService/Public/MagicLeapMusicServiceTypes.h b/Engine/Plugins/Lumin/MagicLeapMusicService/Source/MagicLeapMusicService/Public/MagicLeapMusicServiceTypes.h
new file mode 100644
--- /dev/null
+++ b/Engine/Plugins/Lumin/MagicLeapMusicService/Source/MagicLeapMusicService/Public/MagicLeapMusicServiceTypes.h
@@ -0,0 +1,166 @@
+// Copyright 1998-2019 Epic Games, Inc. All Rights Reserved.
+
+#pragma once
+
+#include "Engine/Engine.h"
+#include "MagicLeapMusicServiceTypes.generated.h"
+
+UENUM(BlueprintType)
+enum class EMagicLeapMusicServiceProviderError : uint8
+{
+	/** Music service has not reported any errors */
+	None,
+	/** Music service has reported a connectivity error */
+	Connectivity,
+	/** Music service has reported a timeout */
+	Timeout,
+	/** Music service has reported a general playback error */
+	GeneralPlayback,
+	/** Music service has reported an app privilege error */
+	Privilege,
+	/** Music service has reported an error specific to it's implementation */
+	ServiceSpecific,
+	/** Music service has reported it has no available memory */
+	NoMemory,
+	/** Music service has reported an unspecified error */
+	Unspecified
+};
+
+UENUM(BlueprintType)
+enum class EMagicLeapMusicServicePlaybackState : uint8
+{
+	/** Music service is currently playing */
+	Playing,
+	/** Music service is currently paused */
+	Paused,
+	/** Music service is currently stopped */
+	Stopped,
+	/** Music service has errored */
+	Error,
+	/** Playback state of the music service is not known */
+	Unknown
+};
+
+UENUM(BlueprintType)
+enum class EMagicLeapMusicServicePlaybackShuffleState : uint8
+{
+	/** Enable shuffling */
+	On,
+	/** Disable shuffling */
+	Off,
+	/** Shuffle state of music service is not known */
+	Unknown
+};
+
+UENUM(BlueprintType)
+enum class EMagicLeapMusicServicePlaybackRepeatState : uint8
+{
+	/** Disable playback repetition */
+	Off,
+	/** Enable single track playback repetition */
+	Song,
+	/** Enable playlist playback repetition */
+	Album,
+	/** Repeat state of music service is not known */
+	Unkown
+};
+
+UENUM(BlueprintType)
+enum class EMagicLeapMusicServiceStatus : uint8
+{
+	/** The music service provider's context has changed */
+	ContextChanged,
+	/** The music service provider has successfully been created */
+	Created,
+	/** Client has successfully logged into the connected music service provider */
+	LoggedIn,
+	/** Client has successfully logged out of the connected music service provider */
+	LoggedOut,
+	/** Music service provider has advanced the current track */
+	NextTrack,
+	/** Music service provider has rewinded the current track */
+	PreviousTrack,
+	/** Music service provider has changed the track */
+	TrackChanged,
+	/** The current status of the music service provider is not known */
+	Unknown
+};
+
+USTRUCT(BlueprintType)
+struct FMagicLeapMusicServiceTrackMetadata
+{
+	GENERATED_BODY()
+
+public:
+	/** Title of the currently selected track */
+	UPROPERTY(VisibleAnywhere, BlueprintReadOnly, Category = "Music Service|MagicLeap")
+	FString TrackTitle;
+
+	/** Name of the album of the currently selected track */
+	UPROPERTY(VisibleAnywhere, BlueprintReadOnly, Category = "Music Service|MagicLeap")
+	FString AlbumName;
+
+	/** URL to the album of the currently selected track */
+	UPROPERTY(VisibleAnywhere, BlueprintReadOnly, Category = "Music Service|MagicLeap")
+	FString AlbumURL;
+
+	/** URL to the album cover of the currently selected track */
+	UPROPERTY(VisibleAnywhere, BlueprintReadOnly, Category = "Music Service|MagicLeap")
+	FString AlbumCoverURL;
+
+	/** Artist name of the currently selected track */
+	UPROPERTY(VisibleAnywhere, BlueprintReadOnly, Category = "Music Service|MagicLeap")
+	FString ArtistName;
+
+	/** Artist URL of the currently selected track */
+	UPROPERTY(VisibleAnywhere, BlueprintReadOnly, Category = "Music Service|MagicLeap")
+	FString ArtistURL;
+
+	/** Runtime length of the currently selected track */
+	UPROPERTY(VisibleAnywhere, BlueprintReadOnly, Category = "Music Service|MagicLeap")
+	FTimespan TrackLength;
+};
+
+
+/** Delegates used for music service callbacks. */
+DECLARE_DYNAMIC_DELEGATE_OneParam(FMagicLeapMusicServicePlaybackStateDelegate, const EMagicLeapMusicServicePlaybackState, PlaybackState);
+DECLARE_DYNAMIC_DELEGATE_OneParam(FMagicLeapMusicServiceRepeatStateDelegate, const EMagicLeapMusicServicePlaybackRepeatState, RepeatState);
+DECLARE_DYNAMIC_DELEGATE_OneParam(FMagicLeapMusicServiceShuffleStateDelegate, const EMagicLeapMusicServicePlaybackShuffleState, ShuffleState);
+DECLARE_DYNAMIC_DELEGATE_OneParam(FMagicLeapMusicServiceMetadataDelegate, const FMagicLeapMusicServiceTrackMetadata&, Metadata);
+DECLARE_DYNAMIC_DELEGATE_OneParam(FMagicLeapMusicServicePositionDelegate, const FTimespan, Position);
+DECLARE_DYNAMIC_DELEGATE_TwoParams(FMagicLeapMusicServiceErrorDelegate, const EMagicLeapMusicServiceProviderError, Error, const int32, ErrorCode);
+DECLARE_DYNAMIC_DELEGATE_OneParam(FMagicLeapMusicServiceStatusDelegate, const EMagicLeapMusicServiceStatus, Status);
+DECLARE_DYNAMIC_DELEGATE_OneParam(FMagicLeapMusicServiceVolumeDelegate, const float, Volume);
+
+
+USTRUCT(BlueprintType)
+struct FMagicLeapMusicServiceCallbacks
+{
+	GENERATED_BODY()
+
+public:
+
+	UPROPERTY(EditAnywhere, BlueprintReadWrite, Category = "Music Service|MagicLeap")
+	FMagicLeapMusicServicePlaybackStateDelegate PlaybackStateDelegate;
+
+	UPROPERTY(EditAnywhere, BlueprintReadWrite, Category = "Music Service|MagicLeap")
+	FMagicLeapMusicServiceRepeatStateDelegate RepeatStateDelegate;
+
+	UPROPERTY(EditAnywhere, BlueprintReadWrite, Category = "Music Service|MagicLeap")
+	FMagicLeapMusicServiceShuffleStateDelegate ShuffleStateDelegate;
+
+	UPROPERTY(EditAnywhere, BlueprintReadWrite, Category = "Music Service|MagicLeap")
+	FMagicLeapMusicServiceMetadataDelegate MetadataDelegate;
+
+	UPROPERTY(EditAnywhere, BlueprintReadWrite, Category = "Music Service|MagicLeap")
+	FMagicLeapMusicServicePositionDelegate PositionDelegate;
+
+	UPROPERTY(EditAnywhere, BlueprintReadWrite, Category = "Music Service|MagicLeap")
+	FMagicLeapMusicServiceErrorDelegate ErrorDelegate;
+
+	UPROPERTY(EditAnywhere, BlueprintReadWrite, Category = "Music Service|MagicLeap")
+	FMagicLeapMusicServiceStatusDelegate StatusDelegate;
+
+	UPROPERTY(EditAnywhere, BlueprintReadWrite, Category = "Music Service|MagicLeap")
+	FMagicLeapMusicServiceVolumeDelegate VolumeDelegate;
+};
diff --git a/Engine/Plugins/Lumin/MagicLeapNetworking/Config/BaseMagicLeapNetworking.ini b/Engine/Plugins/Lumin/MagicLeapNetworking/Config/BaseMagicLeapNetworking.ini
new file mode 100644
--- /dev/null
+++ b/Engine/Plugins/Lumin/MagicLeapNetworking/Config/BaseMagicLeapNetworking.ini
@@ -0,0 +1,6 @@
+[CoreRedirects]
++StructRedirects=(OldName="/Script/MagicLeapNetworking.NetworkingWiFiData", NewName="/Script/MagicLeapNetworking.MagicLeapNetworkingWiFiData")
++ClassRedirects=(OldName="/Script/MagicLeapNetworking.InternetConnectionStatusDelegate", NewName="/Script/MagicLeapNetworking.MagicLeapInternetConnectionStatusDelegate")
++ClassRedirects=(OldName="/Script/MagicLeapNetworking.InternetConnectionStatusDelegateMulti", NewName="/Script/MagicLeapNetworking.MagicLeapInternetConnectionStatusDelegateMulti")
++ClassRedirects=(OldName="/Script/MagicLeapNetworking.WifiStatusDelegate", NewName="/Script/MagicLeapNetworking.MagicLeapWifiStatusDelegate")
++ClassRedirects=(OldName="/Script/MagicLeapNetworking.WifiStatusDelegateMulti", NewName="/Script/MagicLeapNetworking.MagicLeapWifiStatusDelegateMulti")
diff --git a/Engine/Plugins/Lumin/MagicLeapNetworking/MagicLeapNetworking.uplugin b/Engine/Plugins/Lumin/MagicLeapNetworking/MagicLeapNetworking.uplugin
new file mode 100644
--- /dev/null
+++ b/Engine/Plugins/Lumin/MagicLeapNetworking/MagicLeapNetworking.uplugin
@@ -0,0 +1,45 @@
+{
+	"FileVersion" : 3,
+	"Version" : 1,
+	"VersionName": "1.0",
+	"FriendlyName" : "Magic Leap Networking",
+	"CreatedBy" : "Magic Leap, Inc.",
+	"CreatedByURL" : "http://magicleap.com",
+	"Description" : "Support for the Magic Leap Networking API.",
+	"Category" : "Lumin",
+	"EnabledByDefault" : false,
+	"DocsURL" : "",
+	"MarketplaceURL" : "",
+	"SupportURL" : "",
+	"CanContainContent" : false,
+	"IsBetaVersion" : false,
+	"Installed" : false,
+	"SupportedTargetPlatforms": [
+		"Lumin",
+		"Win64",
+		"Mac"
+	],
+	"Modules" : [
+		{
+			"Name" : "MagicLeapNetworking",
+			"Type" : "Runtime",
+			"LoadingPhase" : "PostConfigInit",
+			"WhitelistPlatforms" :
+			[
+				"Win64",
+				"Mac",
+				"Lumin"
+			]
+		}
+	],
+	"Plugins": [
+		{
+			"Name": "MLSDK",
+			"Enabled": true
+		},
+		{
+			"Name": "MagicLeap",
+			"Enabled": true
+		}
+	]
+}
diff --git a/Engine/Plugins/Lumin/MagicLeapNetworking/Source/MagicLeapNetworking.Build.cs b/Engine/Plugins/Lumin/MagicLeapNetworking/Source/MagicLeapNetworking.Build.cs
new file mode 100644
--- /dev/null
+++ b/Engine/Plugins/Lumin/MagicLeapNetworking/Source/MagicLeapNetworking.Build.cs
@@ -0,0 +1,32 @@
+// Copyright 1998-2018 Epic Games, Inc. All Rights Reserved.
+
+using System.IO;
+
+namespace UnrealBuildTool.Rules
+{
+	public class MagicLeapNetworking : ModuleRules
+	{
+		public MagicLeapNetworking( ReadOnlyTargetRules Target ) : base(Target)
+		{
+			PublicIncludePaths.Add(Path.Combine(ModuleDirectory, "Public"));
+
+			PublicDependencyModuleNames.AddRange(
+				new string[]
+				{
+					"Core",
+					"CoreUObject",
+					"Engine"
+				}
+			);
+
+			PrivateDependencyModuleNames.AddRange(
+				new string[]
+				{
+					"LuminRuntimeSettings",
+					"MLSDK",
+					"MagicLeap",
+				}
+			);
+		}
+	}
+}
diff --git a/Engine/Plugins/Lumin/MagicLeapNetworking/Source/Private/MagicLeapNetworkingComponent.cpp b/Engine/Plugins/Lumin/MagicLeapNetworking/Source/Private/MagicLeapNetworkingComponent.cpp
new file mode 100644
--- /dev/null
+++ b/Engine/Plugins/Lumin/MagicLeapNetworking/Source/Private/MagicLeapNetworkingComponent.cpp
@@ -0,0 +1,14 @@
+// Copyright 1998-2018 Epic Games, Inc. All Rights Reserved.
+
+#include "MagicLeapNetworkingComponent.h"
+#include "MagicLeapNetworkingPlugin.h"
+
+bool UMagicLeapNetworkingComponent::IsInternetConnectedAsync()
+{
+	return GetMagicLeapNetworkingPlugin().IsInternetConnectedAsync(ConnectionQueryResultDeleage);
+}
+
+bool UMagicLeapNetworkingComponent::GetWiFiDataAsync()
+{
+	return GetMagicLeapNetworkingPlugin().GetWiFiDataAsync(WifiDataQueryResultDelegate);
+}
diff --git a/Engine/Plugins/Lumin/MagicLeapNetworking/Source/Private/MagicLeapNetworkingFunctionLibrary.cpp b/Engine/Plugins/Lumin/MagicLeapNetworking/Source/Private/MagicLeapNetworkingFunctionLibrary.cpp
new file mode 100644
--- /dev/null
+++ b/Engine/Plugins/Lumin/MagicLeapNetworking/Source/Private/MagicLeapNetworkingFunctionLibrary.cpp
@@ -0,0 +1,18 @@
+// Copyright 1998-2018 Epic Games, Inc. All Rights Reserved.
+
+#include "MagicLeapNetworkingFunctionLibrary.h"
+#include "MagicLeapNetworkingPlugin.h"
+
+bool UMagicLeapNetworkingFunctionLibrary::IsInternetConnectedAsync(const FMagicLeapInternetConnectionStatusDelegate& ResultDelegate)
+{
+	FInternetConnectionStatusDelegateMulti ConnectionQueryResultDelegate;
+	ConnectionQueryResultDelegate.Add(ResultDelegate);
+	return GetMagicLeapNetworkingPlugin().IsInternetConnectedAsync(ConnectionQueryResultDelegate);
+}
+
+bool UMagicLeapNetworkingFunctionLibrary::GetWiFiDataAsync(const FMagicLeapWifiStatusDelegate& ResultDelegate)
+{
+	FWifiStatusDelegateMulti WifiStatusDelegate;
+	WifiStatusDelegate.Add(ResultDelegate);
+	return GetMagicLeapNetworkingPlugin().GetWiFiDataAsync(WifiStatusDelegate);
+}
diff --git a/Engine/Plugins/Lumin/MagicLeapNetworking/Source/Private/MagicLeapNetworkingPlugin.cpp b/Engine/Plugins/Lumin/MagicLeapNetworking/Source/Private/MagicLeapNetworkingPlugin.cpp
new file mode 100644
--- /dev/null
+++ b/Engine/Plugins/Lumin/MagicLeapNetworking/Source/Private/MagicLeapNetworkingPlugin.cpp
@@ -0,0 +1,100 @@
+// Copyright 1998-2018 Epic Games, Inc. All Rights Reserved.
+
+#include "MagicLeapNetworkingPlugin.h"
+#include "Async/Async.h"
+
+DEFINE_LOG_CATEGORY(LogMagicLeapNetworking);
+
+FMagicLeapNetworkingPlugin::FMagicLeapNetworkingPlugin()
+{}
+
+void FMagicLeapNetworkingPlugin::StartupModule()
+{
+	IMagicLeapNetworkingPlugin::StartupModule();
+	TickDelegate = FTickerDelegate::CreateRaw(this, &FMagicLeapNetworkingPlugin::Tick);
+	TickDelegateHandle = FTicker::GetCoreTicker().AddTicker(TickDelegate);
+}
+
+void FMagicLeapNetworkingPlugin::ShutdownModule()
+{
+	FTicker::GetCoreTicker().RemoveTicker(TickDelegateHandle);
+	IMagicLeapNetworkingPlugin::ShutdownModule();
+}
+
+bool FMagicLeapNetworkingPlugin::Tick(float DeltaTime)
+{
+	FConnectionQuery ConnectionQuery;
+	if (ConnectionQueries.Peek(ConnectionQuery))
+	{
+		ConnectionQueries.Pop();
+		ConnectionQuery.ResultDelegate.Broadcast(ConnectionQuery.bConnected);
+	}
+
+	FWifiDataQuery WifiDataQuery;
+	if (WifiDataQueries.Peek(WifiDataQuery))
+	{
+		WifiDataQueries.Pop();
+		WifiDataQuery.ResultDelegate.Broadcast(WifiDataQuery.WifiData);
+	}
+
+	return true;
+}
+
+bool FMagicLeapNetworkingPlugin::IsInternetConnectedAsync(const FInternetConnectionStatusDelegateMulti& ResultDelegate)
+{
+#if WITH_MLSDK
+	AsyncTask(ENamedThreads::AnyBackgroundThreadNormalTask, [this, ResultDelegate]()
+	{
+		bool bIsConnected = false;
+		MLResult Result = MLNetworkingIsInternetConnected(&bIsConnected);
+		if (Result != MLResult_Ok)
+		{
+			UE_LOG(LogMagicLeapNetworking, Error, TEXT("MLNetworkingIsInternetConnected failed with error '%s'"), UTF8_TO_TCHAR(MLNetworkingGetResultString(Result)));
+			return;
+		}
+
+		ConnectionQueries.Enqueue({ bIsConnected, ResultDelegate });
+	});
+	
+	return true;
+#else
+	return false;
+#endif // WITH_MLSDK
+}
+
+bool FMagicLeapNetworkingPlugin::GetWiFiDataAsync(const FWifiStatusDelegateMulti& ResultDelegate)
+{
+#if WITH_MLSDK
+	AsyncTask(ENamedThreads::AnyBackgroundThreadNormalTask, [this, ResultDelegate]()
+	{
+		MLNetworkingWiFiData WifiData;
+		MLNetworkingWiFiDataInit(&WifiData);
+		MLResult Result = MLNetworkingGetWiFiData(&WifiData);
+		if (Result != MLResult_Ok)
+		{
+			UE_LOG(LogMagicLeapNetworking, Error, TEXT("MLNetworkingGetWiFiData failed with error '%s'"), UTF8_TO_TCHAR(MLNetworkingGetResultString(Result)));
+			return;
+		}
+
+		FMagicLeapNetworkingWiFiData UENetworkWifiData;
+		MLToUENetworkingWifiData(WifiData, UENetworkWifiData);
+		WifiDataQueries.Enqueue({ UENetworkWifiData, ResultDelegate });
+	});
+
+
+	return true;
+#else
+	return false;
+#endif // WITH_MLSDK
+}
+
+#if WITH_MLSDK
+void FMagicLeapNetworkingPlugin::MLToUENetworkingWifiData(const  MLNetworkingWiFiData& InMLNetworkingWiFiData, FMagicLeapNetworkingWiFiData& OutUENetworkingWiFiData)
+{
+	OutUENetworkingWiFiData.RSSI = InMLNetworkingWiFiData.rssi;
+	OutUENetworkingWiFiData.Linkspeed = InMLNetworkingWiFiData.linkspeed;
+	OutUENetworkingWiFiData.Frequency = InMLNetworkingWiFiData.frequency;
+}
+#endif // WITH_MLSDK
+
+IMPLEMENT_MODULE(FMagicLeapNetworkingPlugin, MagicLeapNetworking);
diff --git a/Engine/Plugins/Lumin/MagicLeapNetworking/Source/Private/MagicLeapNetworkingPlugin.h b/Engine/Plugins/Lumin/MagicLeapNetworking/Source/Private/MagicLeapNetworkingPlugin.h
new file mode 100644
--- /dev/null
+++ b/Engine/Plugins/Lumin/MagicLeapNetworking/Source/Private/MagicLeapNetworkingPlugin.h
@@ -0,0 +1,50 @@
+// Copyright 1998-2018 Epic Games, Inc. All Rights Reserved.
+
+#pragma once
+
+#include "CoreMinimal.h"
+#include "Containers/Queue.h"
+#include "Containers/Ticker.h"
+#include "IMagicLeapPlugin.h"
+#include "IMagicLeapNetworkingPlugin.h"
+#include "MagicLeapPluginUtil.h"
+#include "MagicLeapNetworkingTypes.h"
+#include "Lumin/CAPIShims/LuminAPINetworking.h"
+
+DECLARE_LOG_CATEGORY_EXTERN(LogMagicLeapNetworking, Verbose, All);
+
+class FMagicLeapNetworkingPlugin : public IMagicLeapNetworkingPlugin
+{
+public:
+	FMagicLeapNetworkingPlugin();
+	void StartupModule() override;
+	void ShutdownModule() override;
+	bool Tick(float DeltaTime);
+	bool IsInternetConnectedAsync(const FInternetConnectionStatusDelegateMulti& ResultDelegate);
+	bool GetWiFiDataAsync(const FWifiStatusDelegateMulti& ResultDelegate);
+
+private:
+	FTickerDelegate TickDelegate;
+	FDelegateHandle TickDelegateHandle;
+	struct FConnectionQuery
+	{
+		bool bConnected;
+		FInternetConnectionStatusDelegateMulti ResultDelegate;
+	};
+	TQueue<FConnectionQuery, EQueueMode::Spsc> ConnectionQueries;
+	struct FWifiDataQuery
+	{
+		FMagicLeapNetworkingWiFiData WifiData;
+		FWifiStatusDelegateMulti ResultDelegate;
+	};
+	TQueue<FWifiDataQuery, EQueueMode::Spsc> WifiDataQueries;
+
+#if WITH_MLSDK
+	void MLToUENetworkingWifiData(const MLNetworkingWiFiData& InMLNetworkingWiFiData, FMagicLeapNetworkingWiFiData& OutUENetworkingWiFiData);
+#endif // WITH_MLSDK
+};
+
+inline FMagicLeapNetworkingPlugin& GetMagicLeapNetworkingPlugin()
+{
+	return FModuleManager::Get().GetModuleChecked<FMagicLeapNetworkingPlugin>("MagicLeapNetworking");
+}
diff --git a/Engine/Plugins/Lumin/MagicLeapNetworking/Source/Public/IMagicLeapNetworkingPlugin.h b/Engine/Plugins/Lumin/MagicLeapNetworking/Source/Public/IMagicLeapNetworkingPlugin.h
new file mode 100644
--- /dev/null
+++ b/Engine/Plugins/Lumin/MagicLeapNetworking/Source/Public/IMagicLeapNetworkingPlugin.h
@@ -0,0 +1,35 @@
+// Copyright 1998-2018 Epic Games, Inc. All Rights Reserved.
+
+#pragma once
+
+#include "Modules/ModuleManager.h"
+
+/**
+ * The public interface to this module.  In most cases, this interface is only public to sibling modules
+ * within this plugin.
+ */
+class IMagicLeapNetworkingPlugin : public IModuleInterface
+{
+public:
+
+	/**
+	 * Singleton-like access to this module's interface.  This is just for convenience!
+	 * Beware of calling this during the shutdown phase, though.  Your module might have been unloaded already.
+	 *
+	 * @return Returns singleton instance, loading the module on demand if needed
+	 */
+	static inline IMagicLeapNetworkingPlugin& Get()
+	{
+		return FModuleManager::LoadModuleChecked<IMagicLeapNetworkingPlugin>("MagicLeapNetworking");
+	}
+
+	/**
+	 * Checks to see if this module is loaded and ready.  It is only valid to call Get() if IsAvailable() returns true.
+	 *
+	 * @return True if the module is loaded and ready to use
+	 */
+	static inline bool IsAvailable()
+	{
+		return FModuleManager::Get().IsModuleLoaded("MagicLeapNetworking");
+	}
+};
diff --git a/Engine/Plugins/Lumin/MagicLeapNetworking/Source/Public/MagicLeapNetworkingComponent.h b/Engine/Plugins/Lumin/MagicLeapNetworking/Source/Public/MagicLeapNetworkingComponent.h
new file mode 100644
--- /dev/null
+++ b/Engine/Plugins/Lumin/MagicLeapNetworking/Source/Public/MagicLeapNetworkingComponent.h
@@ -0,0 +1,42 @@
+// Copyright 1998-2018 Epic Games, Inc. All Rights Reserved.
+
+#pragma once
+
+#include "Engine/Engine.h"
+#include "Components/ActorComponent.h"
+#include "MagicLeapNetworkingTypes.h"
+#include "MagicLeapNetworkingComponent.generated.h"
+
+/**
+	Component that provides access to the Networking API functionality.
+*/
+UCLASS(ClassGroup = MagicLeap, BlueprintType, Blueprintable, EditInlineNew, meta = (BlueprintSpawnableComponent))
+class MAGICLEAPNETWORKING_API UMagicLeapNetworkingComponent : public UActorComponent
+{
+	GENERATED_BODY()
+
+public:
+	/**
+		Asynchronously queries whether or not the internet is connected.  If the delegate broadcasts that the connection is active,
+		it means link layer is up, internet is accessible, and DNS is good.
+		@param ResultDelegate The delegate that will convey the result of the connection query.
+		@return True if the async task is successfully created, false otherwise.
+	*/
+	UFUNCTION(BlueprintCallable, Category = "Networking | MagicLeap")
+	bool IsInternetConnectedAsync();
+
+	/**
+		Asynchronously queries the device's WiFi related data.
+		@param ResultDelegate The delegate that will convey the wifi data.
+		@return True if the async task is successfully created, false otherwise.
+	*/
+	UFUNCTION(BlueprintCallable, Category = "Networking | MagicLeap")
+	bool GetWiFiDataAsync();
+
+private:
+	UPROPERTY(BlueprintAssignable, Category = "Networking | MagicLeap", meta = (AllowPrivateAccess = true))
+	FInternetConnectionStatusDelegateMulti ConnectionQueryResultDeleage;
+
+	UPROPERTY(BlueprintAssignable, Category = "Networking | MagicLeap", meta = (AllowPrivateAccess = true))
+	FWifiStatusDelegateMulti WifiDataQueryResultDelegate;
+};
diff --git a/Engine/Plugins/Lumin/MagicLeapNetworking/Source/Public/MagicLeapNetworkingFunctionLibrary.h b/Engine/Plugins/Lumin/MagicLeapNetworking/Source/Public/MagicLeapNetworkingFunctionLibrary.h
new file mode 100644
--- /dev/null
+++ b/Engine/Plugins/Lumin/MagicLeapNetworking/Source/Public/MagicLeapNetworkingFunctionLibrary.h
@@ -0,0 +1,31 @@
+// Copyright 1998-2018 Epic Games, Inc. All Rights Reserved.
+#pragma once
+
+#include "Engine/Engine.h"
+#include "Kismet/BlueprintFunctionLibrary.h"
+#include "MagicLeapNetworkingTypes.h"
+#include "MagicLeapNetworkingFunctionLibrary.generated.h"
+
+UCLASS(ClassGroup = MagicLeap)
+class MAGICLEAPNETWORKING_API UMagicLeapNetworkingFunctionLibrary : public UBlueprintFunctionLibrary
+{
+	GENERATED_BODY()
+
+public:
+	/**
+		Asynchronously queries whether or not the internet is connected.  If the delegate broadcasts that the connection is active,
+		it means link layer is up, internet is accessible, and DNS is good.
+		@param ResultDelegate The delegate that will convey the result of the connection query.
+		@return True if the async task is successfully created, false otherwise.
+	*/
+	UFUNCTION(BlueprintCallable, Category = "Networking Function Library | MagicLeap")
+	static bool IsInternetConnectedAsync(const FMagicLeapInternetConnectionStatusDelegate& ResultDelegate);
+
+	/** 
+		Asynchronously queries the device's WiFi related data.
+		@param ResultDelegate The delegate that will convey the wifi data.
+		@return True if the async task is successfully created, false otherwise.
+	*/
+	UFUNCTION(BlueprintCallable, Category = "Networking Function Library | MagicLeap")
+	static bool GetWiFiDataAsync(const FMagicLeapWifiStatusDelegate& ResultDelegate);
+};
diff --git a/Engine/Plugins/Lumin/MagicLeapNetworking/Source/Public/MagicLeapNetworkingTypes.h b/Engine/Plugins/Lumin/MagicLeapNetworking/Source/Public/MagicLeapNetworkingTypes.h
new file mode 100644
--- /dev/null
+++ b/Engine/Plugins/Lumin/MagicLeapNetworking/Source/Public/MagicLeapNetworkingTypes.h
@@ -0,0 +1,32 @@
+// Copyright 1998-2018 Epic Games, Inc. All Rights Reserved.
+#pragma once
+
+#include "Engine/Engine.h"
+#include "MagicLeapNetworkingTypes.generated.h"
+
+/** WiFi related data. */
+USTRUCT(BlueprintType)
+struct FMagicLeapNetworkingWiFiData
+{
+	GENERATED_USTRUCT_BODY()
+
+	/** WiFi RSSI in dbM. */
+	UPROPERTY(EditAnywhere, BlueprintReadWrite, Category = "Networking|MagicLeap")
+	int32 RSSI;
+
+	/** WiFi link speed in Mb/s. */
+	UPROPERTY(EditAnywhere, BlueprintReadWrite, Category = "Networking|MagicLeap")
+	int32 Linkspeed;
+
+	/** WiFi frequency in MHz. */
+	UPROPERTY(EditAnywhere, BlueprintReadWrite, Category = "Networking|MagicLeap")
+	float Frequency;
+};
+
+/** Delegate used to convey the result of an internet connection query. */
+DECLARE_DYNAMIC_DELEGATE_OneParam(FMagicLeapInternetConnectionStatusDelegate, bool, bIsConnected);
+DECLARE_DYNAMIC_MULTICAST_DELEGATE_OneParam(FInternetConnectionStatusDelegateMulti, bool, bIsConnected);
+
+/** Delegate used to convey the result of wifi data query. */
+DECLARE_DYNAMIC_DELEGATE_OneParam(FMagicLeapWifiStatusDelegate, FMagicLeapNetworkingWiFiData, WiFiData);
+DECLARE_DYNAMIC_MULTICAST_DELEGATE_OneParam(FWifiStatusDelegateMulti, FMagicLeapNetworkingWiFiData, WiFiData);
diff --git a/Engine/Plugins/Lumin/MagicLeapScreens/Config/BaseMagicLeapScreens.ini b/Engine/Plugins/Lumin/MagicLeapScreens/Config/BaseMagicLeapScreens.ini
new file mode 100644
--- /dev/null
+++ b/Engine/Plugins/Lumin/MagicLeapScreens/Config/BaseMagicLeapScreens.ini
@@ -0,0 +1,7 @@
+[CoreRedirects]
+; Meshing redirects
++StructRedirects=(OldName="/Script/MagicLeapScreens.ScreensWatchHistoryEntry", NewName="/Script/MagicLeapScreens.MagicLeapScreensWatchHistoryEntry")
++StructRedirects=(OldName="/Script/MagicLeapScreens.ScreenTransform", NewName="/Script/MagicLeapScreens.MagicLeapScreenTransform")
++ClassRedirects=(OldName="/Script/MagicLeapScreens.ScreensEntryRequestResultDelegate", NewName="/Script/MagicLeapScreens.MagicLeapScreensEntryRequestResultDelegate")
++ClassRedirects=(OldName="/Script/MagicLeapScreens.ScreensHistoryRequestResultDelegate", NewName="/Script/MagicLeapScreens.MagicLeapScreensHistoryRequestResultDelegate")
++ClassRedirects=(OldName="/Script/MagicLeapScreens.ScreenTransformRequestResultDelegate", NewName="/Script/MagicLeapScreens.MagicLeapScreenTransformRequestResultDelegate")
diff --git a/Engine/Plugins/Lumin/MagicLeapScreens/MagicLeapScreens.uplugin b/Engine/Plugins/Lumin/MagicLeapScreens/MagicLeapScreens.uplugin
--- a/Engine/Plugins/Lumin/MagicLeapScreens/MagicLeapScreens.uplugin
+++ b/Engine/Plugins/Lumin/MagicLeapScreens/MagicLeapScreens.uplugin
@@ -15,7 +15,9 @@
 	"IsBetaVersion" : false,
 	"Installed" : false,
 	"SupportedTargetPlatforms": [
-		"Lumin"
+		"Lumin",
+		"Win64",
+		"Mac"
 	],
 	"Modules" :
 	[
@@ -33,6 +35,10 @@
 	],
 	"Plugins":
 	[
+		{
+			"Name": "MLSDK",
+			"Enabled": true
+		},
 		{
 			"Name": "MagicLeap",
 			"Enabled": true
diff --git a/Engine/Plugins/Lumin/MagicLeapScreens/Source/MagicLeapScreens/MagicLeapScreens.Build.cs b/Engine/Plugins/Lumin/MagicLeapScreens/Source/MagicLeapScreens/MagicLeapScreens.Build.cs
--- a/Engine/Plugins/Lumin/MagicLeapScreens/Source/MagicLeapScreens/MagicLeapScreens.Build.cs
+++ b/Engine/Plugins/Lumin/MagicLeapScreens/Source/MagicLeapScreens/MagicLeapScreens.Build.cs
@@ -28,16 +28,6 @@ namespace UnrealBuildTool.Rules
 					"HeadMountedDisplay",
 				}
 			);
-
-			// This is not ideal but needs to be done in order to expose the private MagicLeapHMD headers to this module.
-			PrivateIncludePaths.Add(Path.Combine(
-				new string[] 
-				{
-					ModuleDirectory,
-					"..", "..", "..",
-					"MagicLeap", "Source", "MagicLeap", "Private"
-				})
-			);
 		}
 	}
 }
diff --git a/Engine/Plugins/Lumin/MagicLeapScreens/Source/MagicLeapScreens/Private/MagicLeapScreensFunctionLibrary.cpp b/Engine/Plugins/Lumin/MagicLeapScreens/Source/MagicLeapScreens/Private/MagicLeapScreensFunctionLibrary.cpp
--- a/Engine/Plugins/Lumin/MagicLeapScreens/Source/MagicLeapScreens/Private/MagicLeapScreensFunctionLibrary.cpp
+++ b/Engine/Plugins/Lumin/MagicLeapScreens/Source/MagicLeapScreens/Private/MagicLeapScreensFunctionLibrary.cpp
@@ -1,51 +1,61 @@
 // Copyright 1998-2019 Epic Games, Inc. All Rights Reserved.
 
 #include "MagicLeapScreensFunctionLibrary.h"
-#include "MagicLeapHMD.h"
-#include "MagicLeapUtils.h"
-#include "MagicLeapScreensWorker.h"
-#include "MagicLeapScreensMsg.h"
 #include "MagicLeapScreensPlugin.h"
 
-void UMagicLeapScreensFunctionLibrary::GetWatchHistoryAsync(const FScreensHistoryRequestResultDelegate& ResultDelegate)
+void UMagicLeapScreensFunctionLibrary::GetWatchHistoryAsync(const FMagicLeapScreensHistoryRequestResultDelegate& ResultDelegate)
 {
 #if WITH_MLSDK
-	FMagicLeapScreensPlugin::GetWatchHistoryEntriesAsync(TOptional<FScreensHistoryRequestResultDelegate>(ResultDelegate));
+	GetMagicLeapScreensPlugin().GetWatchHistoryEntriesAsync(TOptional<FMagicLeapScreensHistoryRequestResultDelegate>(ResultDelegate));
 #else
-	TArray<FScreensWatchHistoryEntry> ResultEntries;
+	TArray<FMagicLeapScreensWatchHistoryEntry> ResultEntries;
 	ResultDelegate.ExecuteIfBound(false, ResultEntries);
 #endif // WITH_MLSDK
 }
 
-void UMagicLeapScreensFunctionLibrary::AddToWatchHistoryAsync(const FScreensWatchHistoryEntry& NewEntry, const FScreensEntryRequestResultDelegate& ResultDelegate)
+void UMagicLeapScreensFunctionLibrary::AddToWatchHistoryAsync(const FMagicLeapScreensWatchHistoryEntry& NewEntry, const FMagicLeapScreensEntryRequestResultDelegate& ResultDelegate)
 {
 #if WITH_MLSDK
-	FMagicLeapScreensPlugin::AddToWatchHistoryAsync(NewEntry, TOptional<FScreensEntryRequestResultDelegate>(ResultDelegate));
+	GetMagicLeapScreensPlugin().AddToWatchHistoryAsync(NewEntry, TOptional<FMagicLeapScreensEntryRequestResultDelegate>(ResultDelegate));
 #else
 	ResultDelegate.ExecuteIfBound(false, NewEntry);
 #endif // WITH_MLSDK
 }
 
-void UMagicLeapScreensFunctionLibrary::UpdateWatchHistoryEntryAsync(const FScreensWatchHistoryEntry& UpdateEntry, const FScreensEntryRequestResultDelegate& ResultDelegate)
+void UMagicLeapScreensFunctionLibrary::UpdateWatchHistoryEntryAsync(const FMagicLeapScreensWatchHistoryEntry& UpdateEntry, const FMagicLeapScreensEntryRequestResultDelegate& ResultDelegate)
 {
 #if WITH_MLSDK
-	FMagicLeapScreensPlugin::UpdateWatchHistoryEntryAsync(UpdateEntry, TOptional<FScreensEntryRequestResultDelegate>(ResultDelegate));
+	GetMagicLeapScreensPlugin().UpdateWatchHistoryEntryAsync(UpdateEntry, TOptional<FMagicLeapScreensEntryRequestResultDelegate>(ResultDelegate));
 #else
 	ResultDelegate.ExecuteIfBound(false, UpdateEntry);
 #endif // WITH_MLSDK
 }
 
-bool UMagicLeapScreensFunctionLibrary::RemoveWatchHistoryEntry(const FScreenID& ID)
+bool UMagicLeapScreensFunctionLibrary::RemoveWatchHistoryEntry(const FGuid& ID)
 {
-	return FMagicLeapScreensPlugin::RemoveWatchHistoryEntry(ID);
+	return GetMagicLeapScreensPlugin().RemoveWatchHistoryEntry(ID);
 }
 
 bool UMagicLeapScreensFunctionLibrary::ClearWatchHistory()
 {
-	return FMagicLeapScreensPlugin::ClearWatchHistory();
+	return GetMagicLeapScreensPlugin().ClearWatchHistory();
 }
 
-bool UMagicLeapScreensFunctionLibrary::GetScreensTransforms(TArray<FScreenTransform>& ScreensTransforms)
+void UMagicLeapScreensFunctionLibrary::UpdateScreenTransformAsync(const FMagicLeapScreenTransform& UpdateTransform, const FMagicLeapScreenTransformRequestResultDelegate& ResultDelegate)
 {
-	return FMagicLeapScreensPlugin::GetScreensTransforms(ScreensTransforms);
+#if WITH_MLSDK
+	GetMagicLeapScreensPlugin().UpdateScreenTransformAsync(UpdateTransform, ResultDelegate);
+#else
+	ResultDelegate.ExecuteIfBound(false);
+#endif // WITH_MLSDK
+}
+
+bool UMagicLeapScreensFunctionLibrary::GetScreenTransform(FMagicLeapScreenTransform& ScreenTransform)
+{
+	return GetMagicLeapScreensPlugin().GetScreenTransform(ScreenTransform);
+}
+
+bool UMagicLeapScreensFunctionLibrary::GetScreensTransforms(TArray<FMagicLeapScreenTransform>& ScreensTransforms)
+{
+	return GetMagicLeapScreensPlugin().GetScreensTransforms(ScreensTransforms);
 }
diff --git a/Engine/Plugins/Lumin/MagicLeapScreens/Source/MagicLeapScreens/Private/MagicLeapScreensPlugin.cpp b/Engine/Plugins/Lumin/MagicLeapScreens/Source/MagicLeapScreens/Private/MagicLeapScreensPlugin.cpp
--- a/Engine/Plugins/Lumin/MagicLeapScreens/Source/MagicLeapScreens/Private/MagicLeapScreensPlugin.cpp
+++ b/Engine/Plugins/Lumin/MagicLeapScreens/Source/MagicLeapScreens/Private/MagicLeapScreensPlugin.cpp
@@ -1,31 +1,26 @@
 // Copyright 1998-2019 Epic Games, Inc. All Rights Reserved.
 
 #include "MagicLeapScreensPlugin.h"
+#include "Async/Async.h"
 #include "RenderUtils.h"
 #include "IMagicLeapPlugin.h"
 #include "HeadMountedDisplayFunctionLibrary.h"
-#include "MagicLeapScreensWorker.h"
+#include "MagicLeapScreensRunnable.h"
 #include "Misc/CoreDelegates.h"
-#include "MagicLeapHMD.h"
- 
+#include "MagicLeapMath.h"
+#include "MagicLeapHandle.h"
+#include "Misc/CommandLine.h"
+
 #define MAX_TEXTURE_SIZE 450 * 450 * 4 // currently limited by binder implementation
 
 DEFINE_LOG_CATEGORY_STATIC(LogScreensPlugin, Display, All);
 
-#if WITH_MLSDK
-MLImage FMagicLeapScreensPlugin::DefaultThumbnail = MLImage();
-FCriticalSection FMagicLeapScreensPlugin::CriticalSection;
-#endif // WITH_MLSDK
-FScreensWorker *FMagicLeapScreensPlugin::Impl = new FScreensWorker();
-TArray<uint8> FMagicLeapScreensPlugin::PixelDataMemPool = TArray<uint8>();
-
 void FMagicLeapScreensPlugin::StartupModule()
 {
 	IModuleInterface::StartupModule();
-	APISetup.Startup();
-	APISetup.LoadDLL(TEXT("ml_screens"));
 
 #if WITH_MLSDK
+	DefaultThumbnail = MLImage();
 	DefaultThumbnail.width = 2;
 	DefaultThumbnail.height = 2;
 	DefaultThumbnail.image_type = MLImageType_RGBA32;
@@ -35,54 +30,33 @@ void FMagicLeapScreensPlugin::StartupModule()
 	FMemory::Memset(DefaultThumbnail.data, 255, DataSize);
 #endif // WITH_MLSDK
 
+	if (!Runnable)
+	{
+		Runnable = new FScreensRunnable();
+	}
 	TickDelegate = FTickerDelegate::CreateRaw(this, &FMagicLeapScreensPlugin::Tick);
 	TickDelegateHandle = FTicker::GetCoreTicker().AddTicker(TickDelegate);
 
 	PixelDataMemPool.Reserve(MAX_TEXTURE_SIZE);
-	bEngineLoopInitComplete = false;
-	FCoreDelegates::OnFEngineLoopInitComplete.AddRaw(this, &FMagicLeapScreensPlugin::OnEngineLoopInitComplete);
 }
 
 void FMagicLeapScreensPlugin::ShutdownModule()
 {
-	FCoreDelegates::OnFEngineLoopInitComplete.RemoveAll(this);
-	APISetup.Shutdown();
-
-#if PLATFORM_LUMIN
-	if (Impl)
+	FScreensRunnable* RunnableToDestroy = Runnable;
+	Runnable = nullptr;
+	AsyncTask(ENamedThreads::AnyBackgroundThreadNormalTask, [RunnableToDestroy]()
 	{
-		Impl->AsyncDestroy();
-		Impl = nullptr;
-	}
-#else
-	delete Impl;
-	Impl = nullptr;
-#endif
-
+		delete RunnableToDestroy;
+	});
 #if WITH_MLSDK
 	delete[] DefaultThumbnail.data;
 	DefaultThumbnail.data = nullptr;
-	MLScreensReleaseWatchHistoryThumbnail(&FMagicLeapScreensPlugin::DefaultThumbnail);
+	MLScreensReleaseWatchHistoryThumbnail(&DefaultThumbnail);
 #endif // WITH_MLSDK
 	FTicker::GetCoreTicker().RemoveTicker(TickDelegateHandle);
 	IModuleInterface::ShutdownModule();
 }
 
-bool FMagicLeapScreensPlugin::IsEngineLoopInitComplete() const
-{
-	return bEngineLoopInitComplete;
-}
-
-void FMagicLeapScreensPlugin::OnEngineLoopInitComplete()
-{
-	if (Impl->Semaphore == nullptr)
-	{
-		Impl->EngineInited();
-	}
-
-	bEngineLoopInitComplete = true;
-}
-
 bool FMagicLeapScreensPlugin::IsSupportedFormat(EPixelFormat InPixelFormat)
 {
 	if (InPixelFormat == PF_B8G8R8A8 || InPixelFormat == PF_R8G8B8A8)
@@ -111,9 +85,9 @@ UTexture2D* FMagicLeapScreensPlugin::MLImageToUTexture2D(const MLImage& Source)
 	return Thumbnail;
 }
 
-void FMagicLeapScreensPlugin::MLWatchHistoryEntryToUnreal(const MLScreensWatchHistoryEntry& InEntry, FScreensWatchHistoryEntry& OutEntry)
+void FMagicLeapScreensPlugin::MLWatchHistoryEntryToUnreal(const MLScreensWatchHistoryEntry& InEntry, FMagicLeapScreensWatchHistoryEntry& OutEntry)
 {
-	OutEntry.ID.ID = InEntry.id;
+	OutEntry.ID = MagicLeap::MLHandleToFGuid(InEntry.id);
 	OutEntry.Title = FString(UTF8_TO_TCHAR(InEntry.title));
 	OutEntry.Subtitle = FString(UTF8_TO_TCHAR(InEntry.subtitle));
 	OutEntry.PlaybackPosition = FTimespan(InEntry.playback_position_ms * ETimespan::TicksPerMillisecond);
@@ -125,7 +99,7 @@ void FMagicLeapScreensPlugin::MLWatchHistoryEntryToUnreal(const MLScreensWatchHi
 	if (Result != MLResult_Ok)
 	{
 		UE_LOG(LogScreensPlugin, Log, TEXT("MLScreensGetWatchHistoryThumbnail failed for screen ID %u with error %s!"), (uint32)InEntry.id, UTF8_TO_TCHAR(MLScreensGetResultString(Result)));
-		OutEntry.Thumbnail = MLImageToUTexture2D(FMagicLeapScreensPlugin::DefaultThumbnail);
+		OutEntry.Thumbnail = MLImageToUTexture2D(DefaultThumbnail);
 	}
 	else
 	{
@@ -136,6 +110,39 @@ void FMagicLeapScreensPlugin::MLWatchHistoryEntryToUnreal(const MLScreensWatchHi
 	}
 }
 
+bool FMagicLeapScreensPlugin::MLScreenInfoToUnreal(const MLScreensScreenInfoEx& InInfo, FMagicLeapScreenTransform& OutScreenTransform, const FTransform& TrackingToWorld, const float WorldToMetersScale)
+{
+	FTransform OutTransform(MagicLeap::ToFMatrix2(InInfo.transform, WorldToMetersScale));
+
+	if (OutTransform.ContainsNaN())
+	{
+		UE_LOG(LogScreensPlugin, Error, TEXT("Screens info entry %d transform contains NaN."), InInfo.screen_id);
+		return false;
+	}
+	if (!OutTransform.GetRotation().IsNormalized())
+	{
+		FQuat OutRotation = OutTransform.GetRotation();
+		OutRotation.Normalize();
+		OutTransform.SetRotation(OutRotation);
+	}
+
+	MLVec3f Scale = MagicLeap::ExtractScale3D(InInfo.transform);
+	MLVec3f Dimensions = InInfo.dimensions;
+	Dimensions.x *= Scale.x;
+	Dimensions.y *= Scale.y;
+	Dimensions.z *= Scale.z;
+
+	OutTransform = OutTransform * TrackingToWorld;
+
+	OutScreenTransform.ID = MagicLeap::MLHandleToFGuid(InInfo.screen_id);
+	OutScreenTransform.VersionNumber = InInfo.version;
+	OutScreenTransform.ScreenPosition = OutTransform.GetLocation();
+	OutScreenTransform.ScreenOrientation = OutTransform.Rotator();
+	OutScreenTransform.ScreenDimensions = MagicLeap::ToFVectorExtents(Dimensions, WorldToMetersScale);
+	OutScreenTransform.ScreenScale3D = MagicLeap::ToFVectorNoScale(Scale);
+	return true;
+}
+
 bool FMagicLeapScreensPlugin::UTexture2DToMLImage(const UTexture2D& Source, MLImage& Target)
 {
 	bool bSuccess = false;
@@ -176,22 +183,25 @@ bool FMagicLeapScreensPlugin::UTexture2DToMLImage(const UTexture2D& Source, MLIm
 	return bSuccess;
 }
 
-bool FMagicLeapScreensPlugin::ShouldUseDefaultThumbnail(const FScreensWatchHistoryEntry& Entry,  MLImage& MLImage)
+bool FMagicLeapScreensPlugin::ShouldUseDefaultThumbnail(const FMagicLeapScreensWatchHistoryEntry& Entry,  MLImage& MLImage)
 {
 	return ((Entry.Thumbnail == nullptr) ||
-		(FMagicLeapScreensPlugin::IsSupportedFormat(Entry.Thumbnail->GetPixelFormat()) == false) ||
-		(FMagicLeapScreensPlugin::UTexture2DToMLImage(*Entry.Thumbnail, MLImage) == false)) ? true : false;
+		(IsSupportedFormat(Entry.Thumbnail->GetPixelFormat()) == false) ||
+		(UTexture2DToMLImage(*Entry.Thumbnail, MLImage) == false)) ? true : false;
 }
 #endif // WITH_MLSDK
 
-bool FMagicLeapScreensPlugin::RemoveWatchHistoryEntry(const FScreenID& ID)
+bool FMagicLeapScreensPlugin::RemoveWatchHistoryEntry(const FGuid& ID)
 {
 #if WITH_MLSDK
 	MLResult Result = MLResult_UnspecifiedFailure;
 	{
-		FScopeLock Lock(&FMagicLeapScreensPlugin::CriticalSection);
-		Result = MLScreensRemoveWatchHistoryEntry(ID.ID);
-		UE_CLOG(Result != MLResult_Ok, LogScreensPlugin, Error, TEXT("MLScreensRemoveWatchHistoryEntry failed with error %d for entry with id %d!"), Result, ID.ID);
+		// TODO: Remove locks once all code is made async
+		FScopeLock Lock(&CriticalSection);
+
+		const int64 EntryID = MagicLeap::FGuidToMLHandle(ID);
+		Result = MLScreensRemoveWatchHistoryEntry(EntryID);
+		UE_CLOG(Result != MLResult_Ok, LogScreensPlugin, Error, TEXT("MLScreensRemoveWatchHistoryEntry failed with error %s for entry with id %d!"), UTF8_TO_TCHAR(MLGetResultString(Result)), EntryID);
 	}
 	return Result == MLResult_Ok;
 #else
@@ -199,83 +209,92 @@ bool FMagicLeapScreensPlugin::RemoveWatchHistoryEntry(const FScreenID& ID)
 #endif // WITH_MLSDK
 }
 
-void FMagicLeapScreensPlugin::GetWatchHistoryEntriesAsync(const TOptional<FScreensHistoryRequestResultDelegate>& OptionalResultDelegate)
+void FMagicLeapScreensPlugin::GetWatchHistoryEntriesAsync(const TOptional<FMagicLeapScreensHistoryRequestResultDelegate>& OptionalResultDelegate)
 {
-	FScreensMessage Msg;
-	Msg.Type = EScreensMsgType::Request;
-	Msg.TaskType = EScreensTaskType::GetHistory;
+	FScreensTask Task;
+	Task.TaskRequestType = FScreensTask::ETaskRequestType::Request;
+	Task.TaskType = FScreensTask::ETaskType::GetHistory;
 	if (OptionalResultDelegate.IsSet())
 	{
-		Msg.HistoryDelegate = OptionalResultDelegate.GetValue();
+		Task.HistoryRequestDelegate = OptionalResultDelegate.GetValue();
 	}
-	FMagicLeapScreensPlugin::Impl->ProcessMessage(Msg);
+	Runnable->PushNewTask(Task);
 }
 
-void FMagicLeapScreensPlugin::AddToWatchHistoryAsync(const FScreensWatchHistoryEntry& NewEntry, const TOptional<FScreensEntryRequestResultDelegate>& OptionalResultDelegate)
+void FMagicLeapScreensPlugin::AddToWatchHistoryAsync(const FMagicLeapScreensWatchHistoryEntry& NewEntry, const TOptional<FMagicLeapScreensEntryRequestResultDelegate>& OptionalResultDelegate)
 {
-	FScreensMessage Msg;
-	Msg.Type = EScreensMsgType::Request;
-	Msg.TaskType = EScreensTaskType::AddToHistory;
+	FScreensTask Task;
+	Task.TaskRequestType = FScreensTask::ETaskRequestType::Request;
+	Task.TaskType = FScreensTask::ETaskType::AddToHistory;
 	if (OptionalResultDelegate.IsSet())
 	{
-		Msg.EntryDelegate = OptionalResultDelegate.GetValue();
+		Task.EntryRequestDelegate = OptionalResultDelegate.GetValue();
 	}
-	Msg.WatchHistory.Add(NewEntry);
-	FMagicLeapScreensPlugin::Impl->ProcessMessage(Msg);
+	Task.WatchHistory.Add(NewEntry);
+	Runnable->PushNewTask(Task);
 }
 
-void FMagicLeapScreensPlugin::UpdateWatchHistoryEntryAsync(const FScreensWatchHistoryEntry& UpdateEntry, const TOptional<FScreensEntryRequestResultDelegate>& OptionalResultDelegate)
+void FMagicLeapScreensPlugin::UpdateWatchHistoryEntryAsync(const FMagicLeapScreensWatchHistoryEntry& UpdateEntry, const TOptional<FMagicLeapScreensEntryRequestResultDelegate>& OptionalResultDelegate)
 {
-	FScreensMessage Msg;
-	Msg.Type = EScreensMsgType::Request;
-	Msg.TaskType = EScreensTaskType::UpdateEntry;
+	FScreensTask Task;
+	Task.TaskRequestType = FScreensTask::ETaskRequestType::Request;
+	Task.TaskType = FScreensTask::ETaskType::UpdateHistoryEntry;
 	if (OptionalResultDelegate.IsSet())
 	{
-		Msg.EntryDelegate = OptionalResultDelegate.GetValue();
+		Task.EntryRequestDelegate = OptionalResultDelegate.GetValue();
 	}
-	Msg.WatchHistory.Add(UpdateEntry);
-	FMagicLeapScreensPlugin::Impl->ProcessMessage(Msg);
+	Task.WatchHistory.Add(UpdateEntry);
+	Runnable->PushNewTask(Task);
 }
 
+void FMagicLeapScreensPlugin::UpdateScreenTransformAsync(const FMagicLeapScreenTransform& UpdateTransform, const FMagicLeapScreenTransformRequestResultDelegate& ResultDelegate)
+{
+	FScreensTask Task;
+	Task.TaskRequestType = FScreensTask::ETaskRequestType::Request;
+	Task.TaskType = FScreensTask::ETaskType::UpdateInfoEntry;
+	Task.TransformRequestDelegate = ResultDelegate;
+	Task.ScreenTransform = UpdateTransform;
+	Runnable->PushNewTask(Task);
+}
 
-FScreensMessage FMagicLeapScreensPlugin::GetWatchHistoryEntries()
+FScreensTask FMagicLeapScreensPlugin::GetWatchHistoryEntries()
 {
-	FScreensMessage Msg;
-	Msg.Type = EScreensMsgType::Response;
-	Msg.TaskType = EScreensTaskType::GetHistory;
-#if WITH_MLSDK
+	FScreensTask Task;
+	Task.TaskRequestType = FScreensTask::ETaskRequestType::Response;
+	Task.TaskType = FScreensTask::ETaskType::GetHistory;
+	#if WITH_MLSDK
 	{
-		FScopeLock Lock(&FMagicLeapScreensPlugin::CriticalSection);
+		FScopeLock Lock(&CriticalSection);
 		MLScreensWatchHistoryList WatchHistoryList;
 		MLResult Result = MLScreensGetWatchHistoryList(&WatchHistoryList);
 		if (Result == MLResult_Ok)
 		{
-			Msg.WatchHistory.Reserve(WatchHistoryList.count);
+			Task.WatchHistory.Reserve(WatchHistoryList.count);
 			for (uint32 i = 0; i < WatchHistoryList.count; ++i)
 			{
-				FScreensWatchHistoryEntry WatchHistoryEntry;
-				FMagicLeapScreensPlugin::MLWatchHistoryEntryToUnreal(WatchHistoryList.entries[i], WatchHistoryEntry);
-				Msg.WatchHistory.Add(WatchHistoryEntry);
+				FMagicLeapScreensWatchHistoryEntry WatchHistoryEntry;
+				MLWatchHistoryEntryToUnreal(WatchHistoryList.entries[i], WatchHistoryEntry);
+				Task.WatchHistory.Add(WatchHistoryEntry);
 			}
 			Result = MLScreensReleaseWatchHistoryList(&WatchHistoryList);
 			UE_CLOG(Result != MLResult_Ok, LogScreensPlugin, Error, TEXT("MLScreensReleaseWatchHistoryList failed with error %s!"), UTF8_TO_TCHAR(MLScreensGetResultString(Result)));
-			Msg.bSuccess = true;
+			Task.bSuccess = true;
 		}
 		else
 		{
-			Msg.bSuccess = false;
+			Task.bSuccess = false;
 			UE_LOG(LogScreensPlugin, Error, TEXT("MLScreensGetWatchHistoryList failed with error %s!"), UTF8_TO_TCHAR(MLScreensGetResultString(Result)));
 		}
 	}
-#endif // WITH_MLSDK
-	return Msg;
+	#endif // WITH_MLSDK
+	return Task;
 }
 
 bool FMagicLeapScreensPlugin::ClearWatchHistory()
 {
 #if WITH_MLSDK
 	{
-		FScopeLock Lock(&FMagicLeapScreensPlugin::CriticalSection);
+		FScopeLock Lock(&CriticalSection);
 		MLScreensWatchHistoryList WatchHistoryList;
 		MLResult Result = MLScreensGetWatchHistoryList(&WatchHistoryList);
 		if (Result == MLResult_Ok)
@@ -303,19 +322,19 @@ bool FMagicLeapScreensPlugin::ClearWatchHistory()
 			return false;
 		}
 	}
-#endif //WITH_MLSDK
-	return true;
+#endif // WITH_MLSDK
+	return false;
 }
 
-FScreensMessage FMagicLeapScreensPlugin::AddToWatchHistory(const FScreensWatchHistoryEntry& WatchHistoryEntry)
+FScreensTask FMagicLeapScreensPlugin::AddToWatchHistory(const FMagicLeapScreensWatchHistoryEntry& WatchHistoryEntry)
 {
-	FScreensMessage Msg;
-	Msg.Type = EScreensMsgType::Response;
-	Msg.TaskType = EScreensTaskType::AddToHistory;
+	FScreensTask Task;
+	Task.TaskRequestType = FScreensTask::ETaskRequestType::Response;
+	Task.TaskType = FScreensTask::ETaskType::AddToHistory;
 
 #if WITH_MLSDK
 	{
-		FScopeLock Lock(&FMagicLeapScreensPlugin::CriticalSection);
+		FScopeLock Lock(&CriticalSection);
 		MLScreensWatchHistoryEntry Entry;
 		Entry.title = TCHAR_TO_UTF8(*WatchHistoryEntry.Title);
 		Entry.subtitle = TCHAR_TO_UTF8(*WatchHistoryEntry.Subtitle);
@@ -324,41 +343,41 @@ FScreensMessage FMagicLeapScreensPlugin::AddToWatchHistory(const FScreensWatchHi
 		Entry.custom_data = TCHAR_TO_UTF8(*WatchHistoryEntry.CustomData);
 		MLImage MLThumbnail;
 
-		if (FMagicLeapScreensPlugin::ShouldUseDefaultThumbnail(WatchHistoryEntry, MLThumbnail))
+		if (ShouldUseDefaultThumbnail(WatchHistoryEntry, MLThumbnail))
 		{
-			MLThumbnail = FMagicLeapScreensPlugin::DefaultThumbnail;
+			MLThumbnail = DefaultThumbnail;
 		}
 
 		MLResult Result = MLScreensInsertWatchHistoryEntry(&Entry, &MLThumbnail);
 		if (Result != MLResult_Ok)
 		{
 			UE_LOG(LogScreensPlugin, Error, TEXT("MLScreensInsertWatchHistoryEntry failed with error %s!"), UTF8_TO_TCHAR(MLScreensGetResultString(Result)));
-			Msg.WatchHistory.Add(WatchHistoryEntry);
-			Msg.bSuccess = false;
+			Task.WatchHistory.Add(WatchHistoryEntry);
+			Task.bSuccess = false;
 		}
 		else
 		{
-			FScreensWatchHistoryEntry NewEntry;
-			FMagicLeapScreensPlugin::MLWatchHistoryEntryToUnreal(Entry, NewEntry);
-			Msg.WatchHistory.Add(NewEntry);
-			Msg.bSuccess = true;
+			FMagicLeapScreensWatchHistoryEntry NewEntry;
+			MLWatchHistoryEntryToUnreal(Entry, NewEntry);
+			Task.WatchHistory.Add(NewEntry);
+			Task.bSuccess = true;
 		}
 	}
 #endif // WITH_MLSDK
-	return Msg;
+	return Task;
 }
 
-FScreensMessage FMagicLeapScreensPlugin::UpdateWatchHistoryEntry(const FScreensWatchHistoryEntry& WatchHistoryEntry)
+FScreensTask FMagicLeapScreensPlugin::UpdateWatchHistoryEntry(const FMagicLeapScreensWatchHistoryEntry& WatchHistoryEntry)
 {
-	FScreensMessage Msg;
-	Msg.Type = EScreensMsgType::Response;
-	Msg.TaskType = EScreensTaskType::UpdateEntry;
+	FScreensTask Task;
+	Task.TaskRequestType = FScreensTask::ETaskRequestType::Response;
+	Task.TaskType = FScreensTask::ETaskType::UpdateHistoryEntry;
 
 #if WITH_MLSDK
 	{
-		FScopeLock Lock(&FMagicLeapScreensPlugin::CriticalSection);
+		FScopeLock Lock(&CriticalSection);
 		MLScreensWatchHistoryEntry Entry;
-		Entry.id = WatchHistoryEntry.ID.ID;
+		Entry.id = MagicLeap::FGuidToMLHandle(WatchHistoryEntry.ID);
 		Entry.title = TCHAR_TO_UTF8(*WatchHistoryEntry.Title);
 		Entry.subtitle = TCHAR_TO_UTF8(*WatchHistoryEntry.Subtitle);
 		Entry.playback_position_ms = static_cast<uint32>(WatchHistoryEntry.PlaybackPosition.GetTotalMilliseconds());
@@ -366,59 +385,82 @@ FScreensMessage FMagicLeapScreensPlugin::UpdateWatchHistoryEntry(const FScreensW
 		Entry.custom_data = TCHAR_TO_UTF8(*WatchHistoryEntry.CustomData);
 		MLImage MLThumbnail;
 
-		if (FMagicLeapScreensPlugin::ShouldUseDefaultThumbnail(WatchHistoryEntry, MLThumbnail))
+		if (ShouldUseDefaultThumbnail(WatchHistoryEntry, MLThumbnail))
 		{
-			MLThumbnail = FMagicLeapScreensPlugin::DefaultThumbnail;
+			MLThumbnail = DefaultThumbnail;
 		}
 
 		MLResult Result = MLScreensUpdateWatchHistoryEntry(&Entry, &MLThumbnail);
 		if (Result != MLResult_Ok)
 		{
 			UE_LOG(LogScreensPlugin, Error, TEXT("MLScreensUpdateWatchHistoryEntry failed with error %s!"), UTF8_TO_TCHAR(MLScreensGetResultString(Result)));
-			Msg.WatchHistory.Add(WatchHistoryEntry);
-			Msg.bSuccess = false;
+			Task.WatchHistory.Add(WatchHistoryEntry);
+			Task.bSuccess = false;
 		}
 		else
 		{
-			FScreensWatchHistoryEntry UpdatedEntry;
+			FMagicLeapScreensWatchHistoryEntry UpdatedEntry;
 			MLWatchHistoryEntryToUnreal(Entry, UpdatedEntry);
-			Msg.WatchHistory.Add(UpdatedEntry);
-			Msg.bSuccess = true;
+			Task.WatchHistory.Add(UpdatedEntry);
+			Task.bSuccess = true;
 		}
 	}
 #endif // WITH_MLSDK
-	return Msg;
+	return Task;
 }
 
 bool FMagicLeapScreensPlugin::Tick(float DeltaTime)
 {
-	if (!FMagicLeapScreensPlugin::Impl->OutgoingMessages.IsEmpty())
+	if (!Runnable->CompletedTaskQueueIsEmpty())
 	{
-		FScreensMessage Msg;
-		FMagicLeapScreensPlugin::Impl->OutgoingMessages.Dequeue(Msg);
+		FScreensTask CompletedTask;
+		Runnable->PeekCompletedTasks(CompletedTask);
+		if (!Runnable->TryGetCompletedTask(CompletedTask))
+		{
+			return true;
+		}
 
-		if (Msg.Type == EScreensMsgType::Request)
+		if (CompletedTask.TaskRequestType == FScreensTask::ETaskRequestType::Request)
 		{
-			UE_LOG(LogScreensPlugin, Error, TEXT("Unexpected EScreensMsgType::Request received from worker thread!"));
+			UE_LOG(LogScreensPlugin, Error, TEXT("Unexpected request received from worker thread!"));
 		}
-		else if (Msg.Type == EScreensMsgType::Response)
+		else if (CompletedTask.TaskRequestType == FScreensTask::ETaskRequestType::Response)
 		{
-			switch (Msg.TaskType)
+			switch (CompletedTask.TaskType)
+			{
+			case FScreensTask::ETaskType::None: break;
+			case FScreensTask::ETaskType::GetHistory:
 			{
-			case EScreensTaskType::None: break;
-			case EScreensTaskType::GetHistory:
+				CompletedTask.HistoryRequestDelegate.ExecuteIfBound(CompletedTask.bSuccess, CompletedTask.WatchHistory);
+			}
+			break;
+			case FScreensTask::ETaskType::AddToHistory:
 			{
-				Msg.HistoryDelegate.ExecuteIfBound(Msg.bSuccess, Msg.WatchHistory);
+				if (CompletedTask.WatchHistory.Num() > 0)
+				{
+					CompletedTask.EntryRequestDelegate.ExecuteIfBound(CompletedTask.bSuccess, CompletedTask.WatchHistory[0]);
+				}
+				else
+				{
+					UE_LOG(LogScreensPlugin, Error, TEXT("Unexpected empty watch history in an AddToHistory response from the worker thread!"));
+				}
 			}
 			break;
-			case EScreensTaskType::AddToHistory:
+			case FScreensTask::ETaskType::UpdateHistoryEntry:
 			{
-				Msg.EntryDelegate.ExecuteIfBound(Msg.bSuccess, Msg.WatchHistory[0]);
+				if (CompletedTask.WatchHistory.Num() > 0)
+				{
+					CompletedTask.EntryRequestDelegate.ExecuteIfBound(CompletedTask.bSuccess, CompletedTask.WatchHistory[0]);
+				}
+				else
+				{
+					UE_LOG(LogScreensPlugin, Error, TEXT("Unexpected empty watch history in an UpdateHistoryEntry response from the worker thread!"));
+				}
 			}
 			break;
-			case EScreensTaskType::UpdateEntry:
+			case FScreensTask::ETaskType::UpdateInfoEntry:
 			{
-				Msg.EntryDelegate.ExecuteIfBound(Msg.bSuccess, Msg.WatchHistory[0]);
+				CompletedTask.TransformRequestDelegate.ExecuteIfBound(CompletedTask.bSuccess);
 			}
 			break;
 			}
@@ -427,60 +469,108 @@ bool FMagicLeapScreensPlugin::Tick(float DeltaTime)
 	return true;
 }
 
-bool FMagicLeapScreensPlugin::GetScreensTransforms(TArray<FScreenTransform>& ScreensTransforms)
+FScreensTask FMagicLeapScreensPlugin::UpdateScreensTransform(const FMagicLeapScreenTransform& InScreenTransform)
 {
+	FScreensTask Task;
+	Task.TaskRequestType = FScreensTask::ETaskRequestType::Response;
+	Task.TaskType = FScreensTask::ETaskType::UpdateInfoEntry;
 #if WITH_MLSDK
-	ScreensTransforms.Empty();
+	const IMagicLeapPlugin& MLPlugin = IMagicLeapPlugin::GetModuleChecked();
+	if (!MLPlugin.IsMagicLeapHMDValid())
+	{
+		return Task;
+	}
 
-	if (!IMagicLeapPlugin::Get().IsMagicLeapHMDValid())
+	float WorldToMetersScale = MLPlugin.GetWorldToMetersScale();
+	const FTransform WorldToTracking = UHeadMountedDisplayFunctionLibrary::GetTrackingToWorldTransform(nullptr).Inverse();
+
+	MLScreensScreenInfoEx ScreensInfo;
+	ScreensInfo.screen_id = MagicLeap::FGuidToMLHandle(InScreenTransform.ID);
+	ScreensInfo.version = InScreenTransform.VersionNumber;
+	ScreensInfo.dimensions = MagicLeap::ToMLVectorExtents(InScreenTransform.ScreenDimensions, WorldToMetersScale);
+	ScreensInfo.transform = MagicLeap::ToMLMat4f(WorldToTracking.TransformPosition(InScreenTransform.ScreenPosition), WorldToTracking.TransformRotation(InScreenTransform.ScreenOrientation.Quaternion()), InScreenTransform.ScreenScale3D, WorldToMetersScale);
+
+	MLResult Result = MLScreensUpdateScreenInfo(&ScreensInfo);
+	UE_CLOG(Result != MLResult_Ok, LogScreensPlugin, Error, TEXT("MLScreensUpdateScreenInfo failed for entry with ID %d with error %s"), ScreensInfo.screen_id, UTF8_TO_TCHAR(MLScreensGetResultString(Result)));
+	Task.ScreenTransform = InScreenTransform;
+	Task.bSuccess = (Result == MLResult_Ok);
+#endif // WITH_MLSDK
+	return Task;
+}
+
+bool FMagicLeapScreensPlugin::GetScreenTransform(FMagicLeapScreenTransform& ScreenTransform)
+{
+	bool bSuccess = false;
+#if WITH_MLSDK
+	const IMagicLeapPlugin& MLPlugin = IMagicLeapPlugin::GetModuleChecked();
+	if (!MLPlugin.IsMagicLeapHMDValid())
 	{
 		return false;
 	}
 
-	const FAppFramework& AppFramework = static_cast<FMagicLeapHMD*>(GEngine->XRSystem->GetHMDDevice())->GetAppFrameworkConst();
-	if (!AppFramework.IsInitialized())
+	float WorldToMetersScale = MLPlugin.GetWorldToMetersScale();
+
+	// ID used to identify a ScreenInfoEx, must be passed via command line
+	FString IDParam;
+	if (!FParse::Value(FCommandLine::Get(), TEXT("screenId="), IDParam))
 	{
+		UE_LOG(LogScreensPlugin, Error, TEXT("A valid Screen ID was not passed via command line!"));
 		return false;
 	}
-	float WorldToMetersScale = AppFramework.GetWorldToMetersScale();
+
+	uint64 ScreenID = FCString::Strtoui64(*IDParam, nullptr, 10);
+	ensure(ScreenID != 0);
+
+	MLScreensScreenInfoEx OutScreenInfo;
+	MLScreensScreenInfoExInit(&OutScreenInfo);
+	MLResult Result = MLScreensGetScreenInfo(ScreenID, &OutScreenInfo);
+	OutScreenInfo.screen_id = ScreenID;
+	if (Result == MLResult_Ok)
+	{
+		const FTransform TrackingToWorld = UHeadMountedDisplayFunctionLibrary::GetTrackingToWorldTransform(nullptr);
+		if (!MLScreenInfoToUnreal(OutScreenInfo, ScreenTransform, TrackingToWorld, WorldToMetersScale))
+		{
+			return false;
+		}
+
+		bSuccess = true;
+	}
+	else
+	{
+		UE_LOG(LogScreensPlugin, Error, TEXT("MLScreensGetScreenInfo failed for screen ID %s with error %s"), *IDParam, UTF8_TO_TCHAR(MLScreensGetResultString(Result)));
+	}
+
+#endif // WITH_MLSDK
+	return bSuccess;
+}
+
+bool FMagicLeapScreensPlugin::GetScreensTransforms(TArray<FMagicLeapScreenTransform>& ScreensTransforms)
+{
+#if WITH_MLSDK
+	ScreensTransforms.Empty();
+
+	const IMagicLeapPlugin& MLPlugin = IMagicLeapPlugin::GetModuleChecked();
+	if (!MLPlugin.IsMagicLeapHMDValid())
+	{
+		return false;
+	}
+
+	float WorldToMetersScale = MLPlugin.GetWorldToMetersScale();
 
 	MLScreensScreenInfoListEx ScreensInfoList;
 	MLScreensScreenInfoListExInit(&ScreensInfoList);
 	MLResult Result = MLScreensGetScreenInfoListEx(&ScreensInfoList);
 	if (Result == MLResult_Ok)
 	{
-		// TODO: Param to GetTrackingToWorldTransform is unused (?)
 		const FTransform TrackingToWorld = UHeadMountedDisplayFunctionLibrary::GetTrackingToWorldTransform(nullptr);
 		for (uint32 i = 0; i < ScreensInfoList.count; ++i)
 		{
 			MLScreensScreenInfoEx& Entry = ScreensInfoList.entries[i];
-			// TODO: Remove this once Screens team fixes setting dimensions when returning the ScreenInfoEx list
-			MLVec3f Scale = MagicLeap::ScaleFromMLMatrix(Entry.transform);
-			Entry.dimensions.x = 0.874f * Scale.x;
-			Entry.dimensions.y = 0.611f * Scale.y;
-			Entry.dimensions.z = 0.5f * Scale.z;
-			FScreenTransform ScreenTransform;
-
-			FTransform EntryTransform = FTransform(FQuat(FVector(0.0, 0.0, 1.0), PI) * MagicLeap::ToFQuat(Entry.transform), MagicLeap::ToFVector(Entry.transform, WorldToMetersScale), FVector(1.0, 1.0, 1.0));
-			if (EntryTransform.ContainsNaN())
+			FMagicLeapScreenTransform ScreenTransform;
+			if (!MLScreenInfoToUnreal(Entry, ScreenTransform, TrackingToWorld, WorldToMetersScale))
 			{
-				UE_LOG(LogScreensPlugin, Error, TEXT("Screens info entry %d transform contains NaN."), Entry.screen_id);
 				continue;
 			}
-			if (!EntryTransform.GetRotation().IsNormalized())
-			{
-				FQuat rotation = EntryTransform.GetRotation();
-				rotation.Normalize();
-				EntryTransform.SetRotation(rotation);
-			}
-			EntryTransform = EntryTransform * TrackingToWorld;
-
-			ScreenTransform.ScreenPosition = EntryTransform.GetLocation();
-			ScreenTransform.ScreenOrientation = EntryTransform.Rotator();
-			ScreenTransform.ScreenDimensions = MagicLeap::ToFVector(Entry.dimensions, WorldToMetersScale);
-			ScreenTransform.ScreenDimensions.X = FMath::Abs<float>(ScreenTransform.ScreenDimensions.X);
-			ScreenTransform.ScreenDimensions.Y = FMath::Abs<float>(ScreenTransform.ScreenDimensions.Y);
-			ScreenTransform.ScreenDimensions.Z = FMath::Abs<float>(ScreenTransform.ScreenDimensions.Z);
 
 			ScreensTransforms.Add(ScreenTransform);
 		}
@@ -500,4 +590,9 @@ bool FMagicLeapScreensPlugin::GetScreensTransforms(TArray<FScreenTransform>& Scr
 	return true;
 }
 
-IMPLEMENT_MODULE(FMagicLeapScreensPlugin, MagicLeapScreens);
\ No newline at end of file
+FScreensRunnable* FMagicLeapScreensPlugin::GetRunnable() const
+{
+	return Runnable;
+}
+
+IMPLEMENT_MODULE(FMagicLeapScreensPlugin, MagicLeapScreens);
